#ttp 309920 - EDITOR: Drag-and-drop an Actor from Scene Outliner to editable Actor variables
#summary Added drag-drop support for actors onto editable blueprint actor properties
#proj Editor
#branch UE4

diff --git a/Engine/Source/Editor/EditorWidgets/Private/SAssetDropTarget.cpp b/Engine/Source/Editor/EditorWidgets/Private/SAssetDropTarget.cpp
--- a/Engine/Source/Editor/EditorWidgets/Private/SAssetDropTarget.cpp
+++ b/Engine/Source/Editor/EditorWidgets/Private/SAssetDropTarget.cpp
@@ -3,6 +3,7 @@
 #include "EditorWidgetsPrivatePCH.h"
 #include "SAssetDropTarget.h"
 #include "AssetDragDropOp.h"
+#include "ActorDragDropOp.h"
 #include "AssetSelection.h"
 
 void SAssetDropTarget::Construct(const FArguments& InArgs )
@@ -136,6 +137,17 @@ UObject* SAssetDropTarget::GetDroppedObject( const FDragDropEvent& DragDropEvent
 			DroppedObject = DroppedAssetData[0].GetAsset();
 		}
 	}
+	// Actor being dragged?
+	else if (Operation->IsOfType<FActorDragDropOp>())
+	{
+		bOutRecognizedEvent = true;
+		TSharedPtr<FActorDragDropOp> ActorDragDrop = StaticCastSharedPtr<FActorDragDropOp>(Operation);
+
+		if (ActorDragDrop->Actors.Num() == 1)
+		{
+			DroppedObject = ActorDragDrop->Actors[0].Get();
+		}
+	}
 
 	return DroppedObject;
 }
