Object name in Details tab now ignores leading and trailing spaces.
#ttp 338977 - EDITOR: Actors can be renamed leading with a space character through the details tab which results in the actor taking on a default generic name
#branch UE4
#proj Editor.EditorWidgets
#add Added call to FText::TrimPrecedingAndTrailing in SObjectNameEditableTextBox::OnNameTextCommitted and SObjectNameEditableTextBox::OnTextChanged.

diff --git a/Engine/Source/Editor/EditorWidgets/Private/SObjectNameEditableTextBox.cpp b/Engine/Source/Editor/EditorWidgets/Private/SObjectNameEditableTextBox.cpp
--- a/Engine/Source/Editor/EditorWidgets/Private/SObjectNameEditableTextBox.cpp
+++ b/Engine/Source/Editor/EditorWidgets/Private/SObjectNameEditableTextBox.cpp
@@ -227,7 +227,9 @@ void SObjectNameEditableTextBox::OnNameTextCommitted(const FText& NewText, EText
 	// Enter being pressed, in which case we will already have been here once with a TextCommit type of OnEnter.
 	if (InTextCommit != ETextCommit::OnCleared)
 	{
-		if (!NewText.IsEmpty())
+		FText TrimmedText = FText::TrimPrecedingAndTrailing(NewText);
+
+		if (!TrimmedText.IsEmpty())
 		{
 			if (Objects.Num() == 1)
 			{
@@ -239,7 +241,7 @@ void SObjectNameEditableTextBox::OnNameTextCommitted(const FText& NewText, EText
 
 					if (Actor->IsActorLabelEditable())
 					{
-						Actor->SetActorLabel( NewText.ToString() );
+						Actor->SetActorLabel(TrimmedText.ToString());
 						LastCommittedTime = FSlateApplication::Get().GetCurrentTime();
 					}
 				}
@@ -248,7 +250,7 @@ void SObjectNameEditableTextBox::OnNameTextCommitted(const FText& NewText, EText
 			{
 				const FScopedTransaction Transaction( LOCTEXT("RenameActorsTransaction", "Rename Multiple Actors") );
 
-				UserSetCommonName = NewText.ToString();
+				UserSetCommonName = TrimmedText.ToString();
 
 				for (int32 i=0; i<Objects.Num(); i++)
 				{
@@ -258,7 +260,7 @@ void SObjectNameEditableTextBox::OnNameTextCommitted(const FText& NewText, EText
 						AActor* Actor = (AActor*)Objects[i].Get();
 						if (Actor->IsActorLabelEditable())
 						{
-							Actor->SetActorLabel(NewText.ToString());
+							Actor->SetActorLabel(TrimmedText.ToString());
 							LastCommittedTime = FSlateApplication::Get().GetCurrentTime();
 						}
 					}
@@ -279,11 +281,13 @@ void SObjectNameEditableTextBox::OnNameTextCommitted(const FText& NewText, EText
 
 void SObjectNameEditableTextBox::OnTextChanged( const FText& InLabel )
 {
-	if( InLabel.IsEmpty() )
+	FText TrimmedLabel = FText::TrimPrecedingAndTrailing(InLabel);
+
+	if (TrimmedLabel.IsEmpty())
 	{
 		TextBox->SetError(LOCTEXT( "RenameFailed_LeftBlank", "Names cannot be left blank" ));
 	}
-	else if(InLabel.ToString().Len() >= NAME_SIZE)
+	else if (TrimmedLabel.ToString().Len() >= NAME_SIZE)
 	{
 		FFormatNamedArguments Arguments;
 		Arguments.Add( TEXT("CharCount"), NAME_SIZE );
