Fixed SAssetSearchBox mouse selection. Previously the search box would close before mouse selected items could be passed back to the text box.
[UE-5742]

diff --git a/Engine/Source/Editor/EditorWidgets/Private/SAssetSearchBox.cpp b/Engine/Source/Editor/EditorWidgets/Private/SAssetSearchBox.cpp
--- a/Engine/Source/Editor/EditorWidgets/Private/SAssetSearchBox.cpp
+++ b/Engine/Source/Editor/EditorWidgets/Private/SAssetSearchBox.cpp
@@ -173,8 +173,11 @@ void SAssetSearchBox::HandleTextCommitted(const FText& NewText, ETextCommit::Typ
 	SetText(CommittedText);
 	OnTextCommitted.ExecuteIfBound(CommittedText, CommitType);
 
-	// Always close the suggestion box
-	SuggestionBox->SetIsOpen(false, false);
+	if(CommitType != ETextCommit::Default)
+	{
+		// Clear the suggestion box if the user has navigated away or set their own text.
+		SuggestionBox->SetIsOpen(false, false);
+	}
 }
 
 void SAssetSearchBox::OnSelectionChanged( TSharedPtr<FString> NewValue, ESelectInfo::Type SelectInfo )
