Fixed Clang compile error (STransportControl.h:26:15: error: 'STransportControl::Tick' hides overloaded virtual function [-Werror,-Woverloaded-virtual])

diff --git a/Engine/Source/Editor/EditorWidgets/Private/STransportControl.h b/Engine/Source/Editor/EditorWidgets/Private/STransportControl.h
--- a/Engine/Source/Editor/EditorWidgets/Private/STransportControl.h
+++ b/Engine/Source/Editor/EditorWidgets/Private/STransportControl.h
@@ -21,6 +21,8 @@ public:
 
 	virtual ~STransportControl() {}
 
+	using SWidget::Tick;
+
 	// Begin FTickableObjectBase implementation
 	virtual bool IsTickable() const override;
 	virtual void Tick( float DeltaTime ) override;
