Fixed "Discovering Assets" indicator from being able to be truncated.

diff --git a/Engine/Source/Editor/EditorWidgets/Private/SAssetDiscoveryIndicator.cpp b/Engine/Source/Editor/EditorWidgets/Private/SAssetDiscoveryIndicator.cpp
--- a/Engine/Source/Editor/EditorWidgets/Private/SAssetDiscoveryIndicator.cpp
+++ b/Engine/Source/Editor/EditorWidgets/Private/SAssetDiscoveryIndicator.cpp
@@ -68,6 +68,8 @@ void SAssetDiscoveryIndicator::Construct( const FArguments& InArgs )
 				SNew(STextBlock)
 				.Font( FEditorStyle::GetFontStyle("AssetDiscoveryIndicator.DiscovertingAssetsFont") )
 				.Text( NSLOCTEXT("AssetDiscoveryIndicator", "DiscoveringAssets", "Discovering Assets").ToString() )
+				.AutoWrapText(true)
+				.Justification(ETextJustify::Center)
 			]
 
 			// Progress bar
