Editor: Rename SPropertyEditorNewAsset to SCreateNewAssetFromFactory and move to EditorWidgets so it can be shared

diff --git a/Engine/Source/Editor/EditorWidgets/Private/SCreateNewAssetFromFactory.cpp b/Engine/Source/Editor/EditorWidgets/Private/SCreateNewAssetFromFactory.cpp
new file mode 100644
--- /dev/null
+++ b/Engine/Source/Editor/EditorWidgets/Private/SCreateNewAssetFromFactory.cpp
@@ -0,0 +1,224 @@
+// Copyright 1998-2015 Epic Games, Inc. All Rights Reserved.
+
+#include "EditorWidgetsPrivatePCH.h"
+#include "SCreateNewAssetFromFactory.h"
+#include "Editor/ContentBrowser/Public/ContentBrowserModule.h"
+#include "AssetRegistryModule.h"
+#include "Developer/AssetTools/Public/AssetToolsModule.h"
+
+#define LOCTEXT_NAMESPACE "SCreateNewAssetFromFactory"
+
+UObject* SCreateNewAssetFromFactory::Create(TWeakObjectPtr<UFactory> FactoryPtr)
+{
+	if (!FactoryPtr.IsValid())
+	{
+		return nullptr;
+	}
+
+	FFormatNamedArguments Args;
+	Args.Add(TEXT("AssetType"), FactoryPtr->GetDisplayName());
+
+	TSharedPtr<SWindow> PropertyEditorNewAssetWindow =
+		SNew(SWindow)
+		.Title(FText::Format(LOCTEXT("SelectPath", "Select Path for {AssetType}"), Args))
+		.ToolTipText(FText::Format(LOCTEXT("SelectPathTooltip", "Select the path where the {AssetType} will be created."), Args))
+		.ClientSize(FVector2D(400, 400));
+
+	TSharedPtr<SCreateNewAssetFromFactory> PropertyEditorNewAssetWidget;
+	PropertyEditorNewAssetWindow->SetContent(SAssignNew(PropertyEditorNewAssetWidget, SCreateNewAssetFromFactory, PropertyEditorNewAssetWindow, FactoryPtr));
+
+	GEditor->EditorAddModalWindow(PropertyEditorNewAssetWindow.ToSharedRef());
+
+	check(PropertyEditorNewAssetWidget.IsValid());
+	return PropertyEditorNewAssetWidget->NewAsset;
+}
+
+void SCreateNewAssetFromFactory::Construct(const FArguments& InArgs, TSharedPtr<SWindow> InParentWindow, TWeakObjectPtr<UFactory> InFactoryPtr)
+{
+	bIsReportingError = false;
+	ParentWindowPtr = InParentWindow;
+	FactoryPtr = InFactoryPtr;
+	NewAsset = nullptr;
+	AssetPath = FString("/Game");
+
+	check(FactoryPtr.IsValid());
+
+	FString AssetName = FactoryPtr->GetDefaultNewAssetName();
+
+	FPathPickerConfig PathPickerConfig;
+	PathPickerConfig.DefaultPath = AssetPath;
+	PathPickerConfig.OnPathSelected = FOnPathSelected::CreateRaw(this, &SCreateNewAssetFromFactory::OnSelectAssetPath);
+
+	FContentBrowserModule& ContentBrowserModule = FModuleManager::LoadModuleChecked<FContentBrowserModule>("ContentBrowser");
+
+	ChildSlot
+	[
+		SNew(SVerticalBox)
+		+ SVerticalBox::Slot()
+		.FillHeight(1.0f)
+		[
+			ContentBrowserModule.Get().CreatePathPicker(PathPickerConfig)
+		]
+		+ SVerticalBox::Slot()
+		.AutoHeight()
+		.Padding(0.0f, 2.0f, 0.0f, 4.0f)
+		[
+			SNew(SSeparator)
+		]
+		+ SVerticalBox::Slot()
+		.AutoHeight()
+		.VAlign(VAlign_Center)
+		.Padding(2.0f, 0.0f, 2.0f, 0.0f)
+		[
+			SNew(SHorizontalBox)
+			+ SHorizontalBox::Slot()
+			.VAlign(VAlign_Center)
+			.Padding(0.0f, 0.0f, 16.0f, 0.0f)
+			.FillWidth(0.6f)
+			[
+				SNew(SHorizontalBox)
+				+ SHorizontalBox::Slot()
+				.VAlign(VAlign_Center)
+				.AutoWidth()
+				[
+					SNew(STextBlock)
+					.Text(LOCTEXT("TextboxLabel", "Asset Name:"))
+				]
+				+ SHorizontalBox::Slot()
+				.Padding(4.0f, 0.0f, 0.0f, 0.0f)
+				.VAlign(VAlign_Center)
+				.FillWidth(1.0f)
+				[
+					SAssignNew(FilenameWidget, SEditableTextBox)
+					.Text(FText::FromString(AssetName))
+					.OnTextChanged(this, &SCreateNewAssetFromFactory::OnFilenameChanged)
+				]
+			]
+			+ SHorizontalBox::Slot()
+			.HAlign(HAlign_Right)
+			.VAlign(VAlign_Center)
+			.FillWidth(0.4f)
+			[
+				SNew(SHorizontalBox)
+				+ SHorizontalBox::Slot()
+				.HAlign(HAlign_Right)
+				.Padding(0.0f, 0.0f, 4.0f, 0.0f)
+				.AutoWidth()
+				[
+					SNew(SButton)
+					.VAlign(VAlign_Center)
+					.ContentPadding(FMargin(8, 2, 8, 2))
+					.OnClicked(this, &SCreateNewAssetFromFactory::OnCreateNewAssetClicked)
+					.IsEnabled(this, &SCreateNewAssetFromFactory::IsCreateNewAssetEnabled)
+					[
+						SNew(STextBlock)
+						.Text(LOCTEXT("CreateButtonText", "Create"))
+					]
+				]
+				+ SHorizontalBox::Slot()
+				.HAlign(HAlign_Right)
+				.AutoWidth()
+				[
+					SNew(SButton)
+					.VAlign(VAlign_Center)
+					.ContentPadding(FMargin(8, 2, 8, 2))
+					.OnClicked(this, &SCreateNewAssetFromFactory::OnCancelCreateNewAsset)
+					[
+						SNew(STextBlock)
+						.Text(LOCTEXT("CancelButtonText", "Cancel"))
+					]
+				]
+			]
+		]
+		+ SVerticalBox::Slot()
+		.AutoHeight()
+		.VAlign(VAlign_Center)
+		.Padding(0.0f, 10.0f, 0.0f, 0.0f)
+		[
+			SNullWidget::NullWidget
+		]
+	];
+
+	OnFilenameChanged(FText::FromString(AssetName));
+}
+
+FReply SCreateNewAssetFromFactory::OnCreateNewAssetClicked()
+{
+	TSharedPtr<SWindow> ParentWindow = ParentWindowPtr.Pin();
+	if (ParentWindow.IsValid())
+	{
+		ParentWindow->RequestDestroyWindow();
+	}
+
+	FString AssetName = FilenameWidget->GetText().ToString();
+	FString AssetFullName = AssetPath / AssetName;
+
+	if (FactoryPtr.IsValid())
+	{
+		UFactory* FactoryInstance = DuplicateObject<UFactory>(FactoryPtr.Get(), GetTransientPackage());
+
+		FEditorDelegates::OnConfigureNewAssetProperties.Broadcast(FactoryInstance);
+		if (FactoryInstance->ConfigureProperties())
+		{
+			FAssetToolsModule& AssetToolsModule = FAssetToolsModule::GetModule();
+			NewAsset = AssetToolsModule.Get().CreateAsset(AssetName, AssetPath, FactoryInstance->GetSupportedClass(), FactoryInstance);
+		}
+	}
+
+	return FReply::Handled();
+}
+
+FReply SCreateNewAssetFromFactory::OnCancelCreateNewAsset()
+{
+	TSharedPtr<SWindow> ParentWindow = ParentWindowPtr.Pin();
+	if (ParentWindow.IsValid())
+	{
+		ParentWindow->RequestDestroyWindow();
+	}
+
+	return FReply::Handled();
+}
+
+void SCreateNewAssetFromFactory::OnSelectAssetPath(const FString& Path)
+{
+	AssetPath = Path;
+	OnFilenameChanged(FilenameWidget->GetText());
+}
+
+void SCreateNewAssetFromFactory::OnFilenameChanged(const FText& InNewName)
+{
+	TArray<FAssetData> AssetData;
+	FAssetRegistryModule& AssetRegistryModule = FModuleManager::LoadModuleChecked<FAssetRegistryModule>("AssetRegistry");
+	AssetRegistryModule.Get().GetAssetsByPath(FName(*AssetPath), AssetData);
+
+	FText ErrorText;
+	if (!FEditorFileUtils::IsFilenameValidForSaving(InNewName.ToString(), ErrorText) || !FName(*InNewName.ToString()).IsValidObjectName(ErrorText))
+	{
+		FilenameWidget->SetError(LOCTEXT("AssetNameInvalid", "Asset name is invalid"));
+		bIsReportingError = true;
+		return;
+	}
+	else
+	{
+		// Check to see if the name conflicts
+		for (auto Iter = AssetData.CreateConstIterator(); Iter; ++Iter)
+		{
+			if (Iter->AssetName.ToString() == InNewName.ToString())
+			{
+				FilenameWidget->SetError(LOCTEXT("AssetInUseError", "Asset name already in use!"));
+				bIsReportingError = true;
+				return;
+			}
+		}
+	}
+
+	FilenameWidget->SetError(FText::FromString(TEXT("")));
+	bIsReportingError = false;
+}
+
+bool SCreateNewAssetFromFactory::IsCreateNewAssetEnabled() const
+{
+	return !bIsReportingError;
+}
+
+#undef LOCTEXT_NAMESPACE
diff --git a/Engine/Source/Editor/EditorWidgets/Public/SCreateNewAssetFromFactory.h b/Engine/Source/Editor/EditorWidgets/Public/SCreateNewAssetFromFactory.h
new file mode 100644
--- /dev/null
+++ b/Engine/Source/Editor/EditorWidgets/Public/SCreateNewAssetFromFactory.h
@@ -0,0 +1,54 @@
+// Copyright 1998-2015 Epic Games, Inc. All Rights Reserved.
+
+#pragma once
+
+ 
+class EDITORWIDGETS_API SCreateNewAssetFromFactory : public SCompoundWidget
+{
+public:
+	SLATE_BEGIN_ARGS(SCreateNewAssetFromFactory)
+	{}
+
+	SLATE_END_ARGS()
+
+	static UObject* Create(TWeakObjectPtr<UFactory> FactoryPtr);
+
+	/** Constructs this widget with InArgs */
+	void Construct(const FArguments& InArgs, TSharedPtr<SWindow> InParentWindow, TWeakObjectPtr<UFactory> InFactoryPtr);
+
+private:
+
+ 	/** Callback when the "create asset" button is clicked. */
+ 	FReply OnCreateNewAssetClicked();
+ 
+ 	/** Callback when the selected asset path has changed. */
+ 	void OnSelectAssetPath(const FString& Path);
+ 
+ 	/** Destroys the window when the operation is cancelled. */
+ 	FReply OnCancelCreateNewAsset();
+
+	/** Callback when the user changes the filename for the Blueprint */
+ 	void OnFilenameChanged(const FText& InNewName);
+ 
+ 	/** Callback to see if creating an asset is enabled */
+ 	bool IsCreateNewAssetEnabled() const;
+
+private:
+	/** The window this widget is nested in */
+	TWeakPtr<SWindow> ParentWindowPtr;
+
+	/** The factory used to create a new asset */
+	TWeakObjectPtr<UFactory> FactoryPtr;
+
+ 	/** The selected path to create the asset */
+ 	FString AssetPath;
+ 
+ 	/** Filename textbox widget */
+ 	TSharedPtr<SEditableTextBox> FilenameWidget;
+
+ 	/** True if an error is currently being reported */
+ 	bool bIsReportingError;
+
+	/** Pointer to the new asset */
+	UObject* NewAsset;
+};
diff --git a/Engine/Source/Editor/PropertyEditor/Private/UserInterface/PropertyEditor/SPropertyEditorAsset.cpp b/Engine/Source/Editor/PropertyEditor/Private/UserInterface/PropertyEditor/SPropertyEditorAsset.cpp
--- a/Engine/Source/Editor/PropertyEditor/Private/UserInterface/PropertyEditor/SPropertyEditorAsset.cpp
+++ b/Engine/Source/Editor/PropertyEditor/Private/UserInterface/PropertyEditor/SPropertyEditorAsset.cpp
@@ -2,7 +2,7 @@
 
 #include "PropertyEditorPrivatePCH.h"
 #include "SPropertyEditorAsset.h"
-#include "SPropertyEditorNewAsset.h"
+#include "SCreateNewAssetFromFactory.h"
 #include "PropertyNode.h"
 #include "PropertyEditor.h"
 #include "AssetThumbnail.h"
diff --git a/Engine/Source/Editor/PropertyEditor/Private/UserInterface/PropertyEditor/SPropertyEditorNewAsset.cpp b/Engine/Source/Editor/PropertyEditor/Private/UserInterface/PropertyEditor/SPropertyEditorNewAsset.cpp
deleted file mode 100644
diff --git a/Engine/Source/Editor/PropertyEditor/Private/UserInterface/PropertyEditor/SPropertyEditorNewAsset.h b/Engine/Source/Editor/PropertyEditor/Private/UserInterface/PropertyEditor/SPropertyEditorNewAsset.h
deleted file mode 100644
diff --git a/Engine/Source/Editor/PropertyEditor/Private/UserInterface/PropertyEditor/SPropertyMenuAssetPicker.cpp b/Engine/Source/Editor/PropertyEditor/Private/UserInterface/PropertyEditor/SPropertyMenuAssetPicker.cpp
--- a/Engine/Source/Editor/PropertyEditor/Private/UserInterface/PropertyEditor/SPropertyMenuAssetPicker.cpp
+++ b/Engine/Source/Editor/PropertyEditor/Private/UserInterface/PropertyEditor/SPropertyMenuAssetPicker.cpp
@@ -2,14 +2,13 @@
 
 #include "PropertyEditorPrivatePCH.h"
 #include "SPropertyMenuAssetPicker.h"
-#include "SPropertyEditorNewAsset.h"
+#include "SCreateNewAssetFromFactory.h"
 #include "AssetRegistryModule.h"
 #include "DelegateFilter.h"
 #include "Editor/ContentBrowser/Public/ContentBrowserModule.h"
 #include "PropertyEditorAssetConstants.h"
 #include "EditorStyleSet.h"
 #include "ClassIconFinder.h"
-#include "SPropertyEditorNewAsset.h"
 
 #define LOCTEXT_NAMESPACE "PropertyEditor"
 
@@ -242,7 +241,7 @@ void SPropertyMenuAssetPicker::OnCreateNewAssetSelected(TWeakObjectPtr<UFactory>
 {
 	if (FactoryPtr.IsValid())
 	{
-		UObject* NewAsset = SPropertyEditorNewAsset::Create(FactoryPtr);
+		UObject* NewAsset = SCreateNewAssetFromFactory::Create(FactoryPtr);
 		if (NewAsset != nullptr)
 		{
 			SetValue(NewAsset);
