Deprecated SCreateNewAssetFromFactory.

diff --git a/Engine/Plugins/2D/Paper2D/Source/Paper2DEditor/Paper2DEditor.Build.cs b/Engine/Plugins/2D/Paper2D/Source/Paper2DEditor/Paper2DEditor.Build.cs
--- a/Engine/Plugins/2D/Paper2D/Source/Paper2DEditor/Paper2DEditor.Build.cs
+++ b/Engine/Plugins/2D/Paper2D/Source/Paper2DEditor/Paper2DEditor.Build.cs
@@ -17,10 +17,8 @@ public class Paper2DEditor : ModuleRules
 				"SlateCore",
 				"Engine",
                 "InputCore",
-				"AssetTools",
 				"UnrealEd", // for FAssetEditorManager
 				"KismetWidgets",
-				"EditorWidgets", // for SCreateNewAssetFromFactory
 				"GraphEditor",
 				"Kismet",  // for FWorkflowCentricApplication
 				"PropertyEditor",
@@ -37,8 +35,16 @@ public class Paper2DEditor : ModuleRules
 		PrivateIncludePathModuleNames.AddRange(
 			new string[] {
 				"Settings",
-				"IntroTutorials"
+				"IntroTutorials",
+                "AssetTools"
 			}
 		);
+
+        DynamicallyLoadedModuleNames.AddRange(
+            new string[] {
+                "AssetTools"
+            }
+        );
+
 	}
 }
diff --git a/Engine/Plugins/2D/Paper2D/Source/Paper2DEditor/Private/TileMapEditing/PaperTileMapDetailsCustomization.cpp b/Engine/Plugins/2D/Paper2D/Source/Paper2DEditor/Private/TileMapEditing/PaperTileMapDetailsCustomization.cpp
--- a/Engine/Plugins/2D/Paper2D/Source/Paper2DEditor/Private/TileMapEditing/PaperTileMapDetailsCustomization.cpp
+++ b/Engine/Plugins/2D/Paper2D/Source/Paper2DEditor/Private/TileMapEditing/PaperTileMapDetailsCustomization.cpp
@@ -6,7 +6,7 @@
 #include "DetailCategoryBuilder.h"
 #include "EdModeTileMap.h"
 #include "PropertyEditing.h"
-#include "SCreateNewAssetFromFactory.h"
+#include "AssetToolsModule.h"
 
 #include "STileLayerList.h"
 #include "ScopedTransaction.h"
@@ -225,7 +225,9 @@ FReply FPaperTileMapDetailsCustomization::OnPromoteButtonClicked()
 				// Try promoting the tile map to be an asset (prompts for a name&path, creates a package and then calls the factory, which renames the existing asset and sets RF_Public)
 				UPaperTileMapPromotionFactory* PromotionFactory = NewObject<UPaperTileMapPromotionFactory>();
 				PromotionFactory->AssetToRename = TileMapComponent->TileMap;
-				SCreateNewAssetFromFactory::Create(PromotionFactory);
+
+				FAssetToolsModule& AssetToolsModule = FAssetToolsModule::GetModule();
+				UObject* NewAsset = AssetToolsModule.Get().CreateAsset(PromotionFactory->GetSupportedClass(), PromotionFactory);
 			}
 		}
 	}
diff --git a/Engine/Source/Developer/AssetTools/Private/AssetTools.cpp b/Engine/Source/Developer/AssetTools/Private/AssetTools.cpp
--- a/Engine/Source/Developer/AssetTools/Private/AssetTools.cpp
+++ b/Engine/Source/Developer/AssetTools/Private/AssetTools.cpp
@@ -265,6 +265,50 @@ UObject* FAssetTools::CreateAsset(const FString& AssetName, const FString& Packa
 	return NewObj;
 }
 
+UObject* FAssetTools::CreateAsset(UClass* AssetClass, UFactory* Factory, FName CallingContext)
+{
+	if (Factory != nullptr)
+	{
+		// Determine the starting path. Try to use the most recently used directory
+		FString AssetPath;
+
+		const FString DefaultFilesystemDirectory = FEditorDirectories::Get().GetLastDirectory(ELastDirectory::NEW_ASSET);
+		if (DefaultFilesystemDirectory.IsEmpty() || !FPackageName::TryConvertFilenameToLongPackageName(DefaultFilesystemDirectory, AssetPath))
+		{
+			// No saved path, just use the game content root
+			AssetPath = TEXT("/Game");
+		}
+
+		FString PackageName;
+		FString AssetName;
+		CreateUniqueAssetName(AssetPath / Factory->GetDefaultNewAssetName(), TEXT(""), PackageName, AssetName);
+
+		FSaveAssetDialogConfig SaveAssetDialogConfig;
+		SaveAssetDialogConfig.DialogTitleOverride = LOCTEXT("SaveAssetDialogTitle", "Save Asset As");
+		SaveAssetDialogConfig.DefaultPath = AssetPath;
+		SaveAssetDialogConfig.DefaultAssetName = AssetName;
+		SaveAssetDialogConfig.ExistingAssetPolicy = ESaveAssetDialogExistingAssetPolicy::AllowButWarn;
+
+		FContentBrowserModule& ContentBrowserModule = FModuleManager::LoadModuleChecked<FContentBrowserModule>("ContentBrowser");
+		FString SaveObjectPath = ContentBrowserModule.Get().CreateModalSaveAssetDialog(SaveAssetDialogConfig);
+		if (!SaveObjectPath.IsEmpty())
+		{
+			FEditorDelegates::OnConfigureNewAssetProperties.Broadcast(Factory);
+			if (Factory->ConfigureProperties())
+			{
+				const FString SavePackageName = FPackageName::ObjectPathToPackageName(SaveObjectPath);
+				const FString SaveAssetPath = FPaths::GetPath(SavePackageName);
+				const FString SaveAssetName = FPaths::GetBaseFilename(SavePackageName);
+				FEditorDirectories::Get().SetLastDirectory(ELastDirectory::NEW_ASSET, SaveAssetPath);
+
+				return CreateAsset(SaveAssetName, SaveAssetPath, AssetClass, Factory, CallingContext);
+			}
+		}
+	}
+
+	return nullptr;
+}
+
 UObject* FAssetTools::DuplicateAsset(const FString& AssetName, const FString& PackagePath, UObject* OriginalObject)
 {
 	// Verify the source object
diff --git a/Engine/Source/Developer/AssetTools/Private/AssetTools.h b/Engine/Source/Developer/AssetTools/Private/AssetTools.h
--- a/Engine/Source/Developer/AssetTools/Private/AssetTools.h
+++ b/Engine/Source/Developer/AssetTools/Private/AssetTools.h
@@ -15,6 +15,7 @@ public:
 	virtual TWeakPtr<IAssetTypeActions> GetAssetTypeActionsForClass( UClass* Class ) const override;
 	virtual bool GetAssetActions( const TArray<UObject*>& InObjects, FMenuBuilder& MenuBuilder, bool bIncludeHeading = true ) override;
 	virtual UObject* CreateAsset(const FString& AssetName, const FString& PackagePath, UClass* AssetClass, UFactory* Factory, FName CallingContext = NAME_None) override;
+	virtual UObject* CreateAsset(UClass* AssetClass, UFactory* Factory, FName CallingContext = NAME_None) override;
 	virtual UObject* DuplicateAsset(const FString& AssetName, const FString& PackagePath, UObject* OriginalObject) override;
 	virtual void RenameAssets(const TArray<FAssetRenameData>& AssetsAndNames) const override;
 	virtual TArray<UObject*> ImportAssets(const FString& DestinationPath) override;
diff --git a/Engine/Source/Developer/AssetTools/Public/IAssetTools.h b/Engine/Source/Developer/AssetTools/Public/IAssetTools.h
--- a/Engine/Source/Developer/AssetTools/Public/IAssetTools.h
+++ b/Engine/Source/Developer/AssetTools/Public/IAssetTools.h
@@ -65,6 +65,9 @@ public:
 	 */
 	virtual UObject* CreateAsset(const FString& AssetName, const FString& PackagePath, UClass* AssetClass, UFactory* Factory, FName CallingContext = NAME_None) = 0;
 
+	/** Opens an asset picker dialog and creates an asset with the chosen path */
+	virtual UObject* CreateAsset(UClass* AssetClass, UFactory* Factory, FName CallingContext = NAME_None) = 0;
+
 	/** Creates an asset with the specified name and path. Uses OriginalObject as the duplication source. */
 	virtual UObject* DuplicateAsset(const FString& AssetName, const FString& PackagePath, UObject* OriginalObject) = 0;
 
diff --git a/Engine/Source/Editor/EditorWidgets/Private/SCreateNewAssetFromFactory.cpp b/Engine/Source/Editor/EditorWidgets/Private/SCreateNewAssetFromFactory.cpp
--- a/Engine/Source/Editor/EditorWidgets/Private/SCreateNewAssetFromFactory.cpp
+++ b/Engine/Source/Editor/EditorWidgets/Private/SCreateNewAssetFromFactory.cpp
@@ -2,12 +2,8 @@
 
 #include "EditorWidgetsPrivatePCH.h"
 #include "SCreateNewAssetFromFactory.h"
-#include "Editor/ContentBrowser/Public/ContentBrowserModule.h"
-#include "AssetRegistryModule.h"
 #include "Developer/AssetTools/Public/AssetToolsModule.h"
 
-#define LOCTEXT_NAMESPACE "SCreateNewAssetFromFactory"
-
 UObject* SCreateNewAssetFromFactory::Create(TWeakObjectPtr<UFactory> FactoryPtr)
 {
 	if (!FactoryPtr.IsValid())
@@ -15,210 +11,7 @@ UObject* SCreateNewAssetFromFactory::Create(TWeakObjectPtr<UFactory> FactoryPtr)
 		return nullptr;
 	}
 
-	FFormatNamedArguments Args;
-	Args.Add(TEXT("AssetType"), FactoryPtr->GetDisplayName());
-
-	TSharedPtr<SWindow> PropertyEditorNewAssetWindow =
-		SNew(SWindow)
-		.Title(FText::Format(LOCTEXT("SelectPath", "Select Path for {AssetType}"), Args))
-		.ToolTipText(FText::Format(LOCTEXT("SelectPathTooltip", "Select the path where the {AssetType} will be created."), Args))
-		.ClientSize(FVector2D(400, 400));
-
-	TSharedPtr<SCreateNewAssetFromFactory> PropertyEditorNewAssetWidget;
-	PropertyEditorNewAssetWindow->SetContent(SAssignNew(PropertyEditorNewAssetWidget, SCreateNewAssetFromFactory, PropertyEditorNewAssetWindow, FactoryPtr));
-
-	GEditor->EditorAddModalWindow(PropertyEditorNewAssetWindow.ToSharedRef());
-
-	check(PropertyEditorNewAssetWidget.IsValid());
-	return PropertyEditorNewAssetWidget->NewAsset;
-}
-
-void SCreateNewAssetFromFactory::Construct(const FArguments& InArgs, TSharedPtr<SWindow> InParentWindow, TWeakObjectPtr<UFactory> InFactoryPtr)
-{
-	bIsReportingError = false;
-	ParentWindowPtr = InParentWindow;
-	FactoryPtr = InFactoryPtr;
-	NewAsset = nullptr;
-	AssetPath = FString("/Game");
-
-	check(FactoryPtr.IsValid());
-
-	FString AssetName = FactoryPtr->GetDefaultNewAssetName();
-
-	FPathPickerConfig PathPickerConfig;
-	PathPickerConfig.DefaultPath = AssetPath;
-	PathPickerConfig.OnPathSelected = FOnPathSelected::CreateRaw(this, &SCreateNewAssetFromFactory::OnSelectAssetPath);
-
-	FContentBrowserModule& ContentBrowserModule = FModuleManager::LoadModuleChecked<FContentBrowserModule>("ContentBrowser");
-
-	ChildSlot
-	[
-		SNew(SVerticalBox)
-		+ SVerticalBox::Slot()
-		.FillHeight(1.0f)
-		[
-			ContentBrowserModule.Get().CreatePathPicker(PathPickerConfig)
-		]
-		+ SVerticalBox::Slot()
-		.AutoHeight()
-		.Padding(0.0f, 2.0f, 0.0f, 4.0f)
-		[
-			SNew(SSeparator)
-		]
-		+ SVerticalBox::Slot()
-		.AutoHeight()
-		.VAlign(VAlign_Center)
-		.Padding(2.0f, 0.0f, 2.0f, 0.0f)
-		[
-			SNew(SHorizontalBox)
-			+ SHorizontalBox::Slot()
-			.VAlign(VAlign_Center)
-			.Padding(0.0f, 0.0f, 16.0f, 0.0f)
-			.FillWidth(0.6f)
-			[
-				SNew(SHorizontalBox)
-				+ SHorizontalBox::Slot()
-				.VAlign(VAlign_Center)
-				.AutoWidth()
-				[
-					SNew(STextBlock)
-					.Text(LOCTEXT("TextboxLabel", "Asset Name:"))
-				]
-				+ SHorizontalBox::Slot()
-				.Padding(4.0f, 0.0f, 0.0f, 0.0f)
-				.VAlign(VAlign_Center)
-				.FillWidth(1.0f)
-				[
-					SAssignNew(FilenameWidget, SEditableTextBox)
-					.Text(FText::FromString(AssetName))
-					.OnTextChanged(this, &SCreateNewAssetFromFactory::OnFilenameChanged)
-				]
-			]
-			+ SHorizontalBox::Slot()
-			.HAlign(HAlign_Right)
-			.VAlign(VAlign_Center)
-			.FillWidth(0.4f)
-			[
-				SNew(SHorizontalBox)
-				+ SHorizontalBox::Slot()
-				.HAlign(HAlign_Right)
-				.Padding(0.0f, 0.0f, 4.0f, 0.0f)
-				.AutoWidth()
-				[
-					SNew(SButton)
-					.VAlign(VAlign_Center)
-					.ContentPadding(FMargin(8, 2, 8, 2))
-					.OnClicked(this, &SCreateNewAssetFromFactory::OnCreateNewAssetClicked)
-					.IsEnabled(this, &SCreateNewAssetFromFactory::IsCreateNewAssetEnabled)
-					[
-						SNew(STextBlock)
-						.Text(LOCTEXT("CreateButtonText", "Create"))
-					]
-				]
-				+ SHorizontalBox::Slot()
-				.HAlign(HAlign_Right)
-				.AutoWidth()
-				[
-					SNew(SButton)
-					.VAlign(VAlign_Center)
-					.ContentPadding(FMargin(8, 2, 8, 2))
-					.OnClicked(this, &SCreateNewAssetFromFactory::OnCancelCreateNewAsset)
-					[
-						SNew(STextBlock)
-						.Text(LOCTEXT("CancelButtonText", "Cancel"))
-					]
-				]
-			]
-		]
-		+ SVerticalBox::Slot()
-		.AutoHeight()
-		.VAlign(VAlign_Center)
-		.Padding(0.0f, 10.0f, 0.0f, 0.0f)
-		[
-			SNullWidget::NullWidget
-		]
-	];
-
-	OnFilenameChanged(FText::FromString(AssetName));
-}
-
-FReply SCreateNewAssetFromFactory::OnCreateNewAssetClicked()
-{
-	TSharedPtr<SWindow> ParentWindow = ParentWindowPtr.Pin();
-	if (ParentWindow.IsValid())
-	{
-		ParentWindow->RequestDestroyWindow();
-	}
-
-	FString AssetName = FilenameWidget->GetText().ToString();
-	FString AssetFullName = AssetPath / AssetName;
-
-	if (FactoryPtr.IsValid())
-	{
-		UFactory* FactoryInstance = DuplicateObject<UFactory>(FactoryPtr.Get(), GetTransientPackage());
-
-		FEditorDelegates::OnConfigureNewAssetProperties.Broadcast(FactoryInstance);
-		if (FactoryInstance->ConfigureProperties())
-		{
-			FAssetToolsModule& AssetToolsModule = FAssetToolsModule::GetModule();
-			NewAsset = AssetToolsModule.Get().CreateAsset(AssetName, AssetPath, FactoryInstance->GetSupportedClass(), FactoryInstance);
-		}
-	}
-
-	return FReply::Handled();
-}
-
-FReply SCreateNewAssetFromFactory::OnCancelCreateNewAsset()
-{
-	TSharedPtr<SWindow> ParentWindow = ParentWindowPtr.Pin();
-	if (ParentWindow.IsValid())
-	{
-		ParentWindow->RequestDestroyWindow();
-	}
-
-	return FReply::Handled();
-}
-
-void SCreateNewAssetFromFactory::OnSelectAssetPath(const FString& Path)
-{
-	AssetPath = Path;
-	OnFilenameChanged(FilenameWidget->GetText());
-}
-
-void SCreateNewAssetFromFactory::OnFilenameChanged(const FText& InNewName)
-{
-	TArray<FAssetData> AssetData;
-	FAssetRegistryModule& AssetRegistryModule = FModuleManager::LoadModuleChecked<FAssetRegistryModule>("AssetRegistry");
-	AssetRegistryModule.Get().GetAssetsByPath(FName(*AssetPath), AssetData);
-
-	FText ErrorText;
-	if (!FEditorFileUtils::IsFilenameValidForSaving(InNewName.ToString(), ErrorText) || !FName(*InNewName.ToString()).IsValidObjectName(ErrorText))
-	{
-		FilenameWidget->SetError(LOCTEXT("AssetNameInvalid", "Asset name is invalid"));
-		bIsReportingError = true;
-		return;
-	}
-	else
-	{
-		// Check to see if the name conflicts
-		for (auto Iter = AssetData.CreateConstIterator(); Iter; ++Iter)
-		{
-			if (Iter->AssetName.ToString() == InNewName.ToString())
-			{
-				FilenameWidget->SetError(LOCTEXT("AssetInUseError", "Asset name already in use!"));
-				bIsReportingError = true;
-				return;
-			}
-		}
-	}
-
-	FilenameWidget->SetError(FText::FromString(TEXT("")));
-	bIsReportingError = false;
-}
-
-bool SCreateNewAssetFromFactory::IsCreateNewAssetEnabled() const
-{
-	return !bIsReportingError;
+	FAssetToolsModule& AssetToolsModule = FAssetToolsModule::GetModule();
+	return AssetToolsModule.Get().CreateAsset(FactoryPtr->GetSupportedClass(), FactoryPtr.Get());
 }
 
-#undef LOCTEXT_NAMESPACE
diff --git a/Engine/Source/Editor/EditorWidgets/Public/SCreateNewAssetFromFactory.h b/Engine/Source/Editor/EditorWidgets/Public/SCreateNewAssetFromFactory.h
--- a/Engine/Source/Editor/EditorWidgets/Public/SCreateNewAssetFromFactory.h
+++ b/Engine/Source/Editor/EditorWidgets/Public/SCreateNewAssetFromFactory.h
@@ -2,7 +2,6 @@
 
 #pragma once
 
- 
 class EDITORWIDGETS_API SCreateNewAssetFromFactory : public SCompoundWidget
 {
 public:
@@ -11,44 +10,10 @@ public:
 
 	SLATE_END_ARGS()
 
-	static UObject* Create(TWeakObjectPtr<UFactory> FactoryPtr);
-
-	/** Constructs this widget with InArgs */
-	void Construct(const FArguments& InArgs, TSharedPtr<SWindow> InParentWindow, TWeakObjectPtr<UFactory> InFactoryPtr);
-
-private:
-
- 	/** Callback when the "create asset" button is clicked. */
- 	FReply OnCreateNewAssetClicked();
- 
- 	/** Callback when the selected asset path has changed. */
- 	void OnSelectAssetPath(const FString& Path);
- 
- 	/** Destroys the window when the operation is cancelled. */
- 	FReply OnCancelCreateNewAsset();
-
-	/** Callback when the user changes the filename for the Blueprint */
- 	void OnFilenameChanged(const FText& InNewName);
- 
- 	/** Callback to see if creating an asset is enabled */
- 	bool IsCreateNewAssetEnabled() const;
+	virtual ~SCreateNewAssetFromFactory() {}
 
-private:
-	/** The window this widget is nested in */
-	TWeakPtr<SWindow> ParentWindowPtr;
-
-	/** The factory used to create a new asset */
-	TWeakObjectPtr<UFactory> FactoryPtr;
-
- 	/** The selected path to create the asset */
- 	FString AssetPath;
- 
- 	/** Filename textbox widget */
- 	TSharedPtr<SEditableTextBox> FilenameWidget;
-
- 	/** True if an error is currently being reported */
- 	bool bIsReportingError;
+	DEPRECATED(4.8, "Please do not use SCreateNewAssetFromFactory::Create - use IAssetTools::CreateAsset instead.")
+	static UObject* Create(TWeakObjectPtr<UFactory> FactoryPtr);
 
-	/** Pointer to the new asset */
-	UObject* NewAsset;
+	void Construct(const FArguments& InArgs, TSharedPtr<SWindow> InParentWindow, TWeakObjectPtr<UFactory> InFactoryPtr) {}
 };
diff --git a/Engine/Source/Editor/PropertyEditor/Private/UserInterface/PropertyEditor/SPropertyEditorAsset.cpp b/Engine/Source/Editor/PropertyEditor/Private/UserInterface/PropertyEditor/SPropertyEditorAsset.cpp
--- a/Engine/Source/Editor/PropertyEditor/Private/UserInterface/PropertyEditor/SPropertyEditorAsset.cpp
+++ b/Engine/Source/Editor/PropertyEditor/Private/UserInterface/PropertyEditor/SPropertyEditorAsset.cpp
@@ -2,7 +2,6 @@
 
 #include "PropertyEditorPrivatePCH.h"
 #include "SPropertyEditorAsset.h"
-#include "SCreateNewAssetFromFactory.h"
 #include "PropertyNode.h"
 #include "PropertyEditor.h"
 #include "AssetThumbnail.h"
diff --git a/Engine/Source/Editor/PropertyEditor/Private/UserInterface/PropertyEditor/SPropertyMenuAssetPicker.cpp b/Engine/Source/Editor/PropertyEditor/Private/UserInterface/PropertyEditor/SPropertyMenuAssetPicker.cpp
--- a/Engine/Source/Editor/PropertyEditor/Private/UserInterface/PropertyEditor/SPropertyMenuAssetPicker.cpp
+++ b/Engine/Source/Editor/PropertyEditor/Private/UserInterface/PropertyEditor/SPropertyMenuAssetPicker.cpp
@@ -2,10 +2,10 @@
 
 #include "PropertyEditorPrivatePCH.h"
 #include "SPropertyMenuAssetPicker.h"
-#include "SCreateNewAssetFromFactory.h"
 #include "AssetRegistryModule.h"
+#include "AssetToolsModule.h"
 #include "DelegateFilter.h"
-#include "Editor/ContentBrowser/Public/ContentBrowserModule.h"
+#include "ContentBrowserModule.h"
 #include "PropertyEditorAssetConstants.h"
 #include "EditorStyleSet.h"
 #include "ClassIconFinder.h"
@@ -241,7 +241,9 @@ void SPropertyMenuAssetPicker::OnCreateNewAssetSelected(TWeakObjectPtr<UFactory>
 {
 	if (FactoryPtr.IsValid())
 	{
-		UObject* NewAsset = SCreateNewAssetFromFactory::Create(FactoryPtr);
+		UFactory* FactoryInstance = DuplicateObject<UFactory>(FactoryPtr.Get(), GetTransientPackage());
+		FAssetToolsModule& AssetToolsModule = FAssetToolsModule::GetModule();
+		UObject* NewAsset = AssetToolsModule.Get().CreateAsset(FactoryInstance->GetSupportedClass(), FactoryInstance);
 		if (NewAsset != nullptr)
 		{
 			SetValue(NewAsset);
diff --git a/Engine/Source/Editor/PropertyEditor/PropertyEditor.Build.cs b/Engine/Source/Editor/PropertyEditor/PropertyEditor.Build.cs
--- a/Engine/Source/Editor/PropertyEditor/PropertyEditor.Build.cs
+++ b/Engine/Source/Editor/PropertyEditor/PropertyEditor.Build.cs
@@ -37,6 +37,9 @@ public class PropertyEditor : ModuleRules
 		PrivateIncludePathModuleNames.AddRange(
 			new string[] {
 				"MainFrame",
+                "AssetRegistry",
+                "AssetTools",
+                "ContentBrowser"
 			}
 		);
 
@@ -59,6 +62,7 @@ public class PropertyEditor : ModuleRules
 		DynamicallyLoadedModuleNames.AddRange(
 			new string[] {
                 "AssetRegistry",
+                "AssetTools",
 				"ContentBrowser",
 				"Documentation",
 				"MainFrame",
diff --git a/Engine/Source/Editor/UnrealEd/Public/EditorDirectories.h b/Engine/Source/Editor/UnrealEd/Public/EditorDirectories.h
--- a/Engine/Source/Editor/UnrealEd/Public/EditorDirectories.h
+++ b/Engine/Source/Editor/UnrealEd/Public/EditorDirectories.h
@@ -19,6 +19,7 @@ namespace ELastDirectory
 		WORLD_ROOT,
 		LEVEL,
 		PROJECT,
+		NEW_ASSET,
 		MAX
 	};
 };
