[INTEGRATE] Change 2423049 by [REDACTED_PERSON][REDACTED_HANDLE] on 2015/01/28 19:17:57

	Fixed "Discovering Assets" not centered nicely in Content Browser

diff --git a/Engine/Source/Editor/EditorWidgets/Private/SAssetDiscoveryIndicator.cpp b/Engine/Source/Editor/EditorWidgets/Private/SAssetDiscoveryIndicator.cpp
--- a/Engine/Source/Editor/EditorWidgets/Private/SAssetDiscoveryIndicator.cpp
+++ b/Engine/Source/Editor/EditorWidgets/Private/SAssetDiscoveryIndicator.cpp
@@ -48,37 +48,41 @@ void SAssetDiscoveryIndicator::Construct( const FArguments& InArgs )
 
 	ChildSlot
 	[
-		SNew(SBorder)
-		.BorderImage( FEditorStyle::GetBrush("Menu.Background") )
-		.BorderBackgroundColor(this, &SAssetDiscoveryIndicator::GetBorderBackgroundColor)
-		.ColorAndOpacity(this, &SAssetDiscoveryIndicator::GetIndicatorColorAndOpacity)
-		.DesiredSizeScale(this, &SAssetDiscoveryIndicator::GetIndicatorDesiredSizeScale)
-		.Visibility(this, &SAssetDiscoveryIndicator::GetIndicatorVisibility)
+		SNew(SVerticalBox)
+		+SVerticalBox::Slot()
 		.Padding(InArgs._Padding)
 		.HAlign(HAlign_Center)
-		.VAlign(VAlign_Center)
 		[
-			SNew(SVerticalBox)
-
-			// Text
-			+SVerticalBox::Slot()
-			.AutoHeight()
-			.Padding(0, 4)
-			[
-				SNew(STextBlock)
-				.Font( FEditorStyle::GetFontStyle("AssetDiscoveryIndicator.DiscovertingAssetsFont") )
-				.Text( NSLOCTEXT("AssetDiscoveryIndicator", "DiscoveringAssets", "Discovering Assets") )
-				.AutoWrapText(true)
-				.Justification(ETextJustify::Center)
-			]
-
-			// Progress bar
-			+SVerticalBox::Slot()
-			.AutoHeight()
-			.Padding(0, 4)
+			SNew(SBorder)
+			.BorderImage( FEditorStyle::GetBrush("Menu.Background") )
+			.BorderBackgroundColor(this, &SAssetDiscoveryIndicator::GetBorderBackgroundColor)
+			.ColorAndOpacity(this, &SAssetDiscoveryIndicator::GetIndicatorColorAndOpacity)
+			.DesiredSizeScale(this, &SAssetDiscoveryIndicator::GetIndicatorDesiredSizeScale)
+			.Visibility(this, &SAssetDiscoveryIndicator::GetIndicatorVisibility)
+			.VAlign(VAlign_Center)
 			[
-				SNew(SProgressBar)
-				.Percent( this, &SAssetDiscoveryIndicator::GetProgress )
+				SNew(SVerticalBox)
+
+				// Text
+				+SVerticalBox::Slot()
+				.AutoHeight()
+				.Padding(12, 4)
+				[
+					SNew(STextBlock)
+					.Font( FEditorStyle::GetFontStyle("AssetDiscoveryIndicator.DiscovertingAssetsFont") )
+					.Text( NSLOCTEXT("AssetDiscoveryIndicator", "DiscoveringAssets", "Discovering Assets") )
+					.AutoWrapText(true)
+					.Justification(ETextJustify::Center)
+				]
+
+				// Progress bar
+				+SVerticalBox::Slot()
+				.AutoHeight()
+				.Padding(12, 4)
+				[
+					SNew(SProgressBar)
+					.Percent( this, &SAssetDiscoveryIndicator::GetProgress )
+				]
 			]
 		]
 	];
