[INTEGRATE] Change 2425685 by [REDACTED_PERSON][REDACTED_HANDLE] on 2015/01/30 09:52:10

	Editor - Improving the droptarget for assets in the details panels to be more obvious when a drag begins.

diff --git a/Engine/Source/Editor/EditorStyle/Private/SlateEditorStyle.cpp b/Engine/Source/Editor/EditorStyle/Private/SlateEditorStyle.cpp
--- a/Engine/Source/Editor/EditorStyle/Private/SlateEditorStyle.cpp
+++ b/Engine/Source/Editor/EditorStyle/Private/SlateEditorStyle.cpp
@@ -218,6 +218,9 @@ void FSlateEditorStyle::FStyle::SetupGeneralStyles()
 
 	Set( "PlainBorder", new BORDER_BRUSH( "Common/PlainBorder", 2.f/8.f) );
 
+	Set( "WideDash.Horizontal", new IMAGE_BRUSH("Common/WideDash_Horizontal", FVector2D(22, 4), FLinearColor::White, ESlateBrushTileType::Horizontal));
+	Set( "WideDash.Vertical", new IMAGE_BRUSH("Common/WideDash_Vertical", FVector2D(4, 22), FLinearColor::White, ESlateBrushTileType::Vertical));
+
 	// Debug Colors
 	Set( "MultiboxHookColor", FLinearColor(0.f, 1.f, 0.f, 1.f) );
 
diff --git a/Engine/Source/Editor/EditorWidgets/Private/SAssetDropTarget.cpp b/Engine/Source/Editor/EditorWidgets/Private/SAssetDropTarget.cpp
--- a/Engine/Source/Editor/EditorWidgets/Private/SAssetDropTarget.cpp
+++ b/Engine/Source/Editor/EditorWidgets/Private/SAssetDropTarget.cpp
@@ -1,10 +1,13 @@
-// Copyright 1998-2015 Epic Games, Inc. All Rights Reserved.
+﻿// Copyright 1998-2015 Epic Games, Inc. All Rights Reserved.
 
 #include "EditorWidgetsPrivatePCH.h"
 #include "SAssetDropTarget.h"
 #include "AssetDragDropOp.h"
 #include "ActorDragDropOp.h"
 #include "AssetSelection.h"
+#include "SScaleBox.h"
+
+#define LOCTEXT_NAMESPACE "EditorWidgets"
 
 void SAssetDropTarget::Construct(const FArguments& InArgs )
 {
@@ -13,21 +16,63 @@ void SAssetDropTarget::Construct(const FArguments& InArgs )
 
 	bIsDragEventRecognized = false;
 	bAllowDrop = false;
+	bIsDragOver = false;
 
-	SBorder::Construct( 
-		SBorder::FArguments()
-		.Padding(0)
-		.BorderImage( this, &SAssetDropTarget::GetDragBorder )
-		.BorderBackgroundColor( this, &SAssetDropTarget::GetDropBorderColor )
+	ChildSlot
+	[
+		SNew(SOverlay)
+			
+		+ SOverlay::Slot()
 		[
 			InArgs._Content.Widget
 		]
-	);
+
+		+ SOverlay::Slot()
+		[
+			SNew(SBorder)
+			.Visibility(this, &SAssetDropTarget::GetDragOverlayVisibility)
+			.BorderImage(FEditorStyle::GetBrush("NoBorder"))
+			.BorderBackgroundColor(this, &SAssetDropTarget::GetDropBorderColor)
+			[
+				SNew(SBorder)
+				.BorderImage(FEditorStyle::GetBrush("WhiteBrush"))
+				.BorderBackgroundColor(this, &SAssetDropTarget::GetBackgroundBrightness)
+				.Padding(15.0f)
+				[
+					SNew(SScaleBox)
+					.Stretch(EStretch::ScaleToFit)
+					[
+						SNew(STextBlock)
+						.Text(this, &SAssetDropTarget::GetDragOverlayText)
+						.ColorAndOpacity(this, &SAssetDropTarget::GetDropBorderColor)
+					]
+				]
+			]
+		]
+	];
 }
 
-const FSlateBrush* SAssetDropTarget::GetDragBorder() const
+FSlateColor SAssetDropTarget::GetBackgroundBrightness() const
 {
-	return bIsDragEventRecognized ? FEditorStyle::GetBrush("MaterialList.DragDropBorder") : FEditorStyle::GetBrush("NoBorder");
+	return ( bAllowDrop && bIsDragOver ) ? FLinearColor(1, 1, 1, 0.40f) : FLinearColor(1, 1, 1, 0.25f);
+}
+
+EVisibility SAssetDropTarget::GetDragOverlayVisibility() const
+{
+	if ( FSlateApplication::Get().IsDragDropping() )
+	{
+		if ( AllowDrop(FSlateApplication::Get().GetDragDroppingContent()) || bIsDragOver )
+		{
+			return EVisibility::HitTestInvisible;
+		}
+	}
+
+	return EVisibility::Hidden;
+}
+
+FText SAssetDropTarget::GetDragOverlayText() const
+{
+	return bAllowDrop ? FText::GetEmpty() : LOCTEXT("Nope", "Nope");
 }
 
 FSlateColor SAssetDropTarget::GetDropBorderColor() const
@@ -38,21 +83,27 @@ FSlateColor SAssetDropTarget::GetDropBorderColor() const
 }
 
 FReply SAssetDropTarget::OnDragOver( const FGeometry& MyGeometry, const FDragDropEvent& DragDropEvent )
+{
+	// Handle the reply if we are allowed to drop, otherwise do not handle it.
+	return AllowDrop(DragDropEvent.GetOperation()) ? FReply::Handled() : FReply::Unhandled();
+}
+
+bool SAssetDropTarget::AllowDrop(TSharedPtr<FDragDropOperation> DragDropOperation) const
 {
 	bool bRecognizedEvent = false;
-	UObject* Object = GetDroppedObject( DragDropEvent, bRecognizedEvent );
+	UObject* Object = GetDroppedObject(DragDropOperation, bRecognizedEvent);
 
 	// Not being dragged over by a recognizable event
 	bIsDragEventRecognized = bRecognizedEvent;
 
 	bAllowDrop = false;
 
-	if( Object )
+	if ( Object )
 	{
 		// Check and see if its valid to drop this object
-		if( OnIsAssetAcceptableForDrop.IsBound() )
+		if ( OnIsAssetAcceptableForDrop.IsBound() )
 		{
-			bAllowDrop = OnIsAssetAcceptableForDrop.Execute( Object );
+			bAllowDrop = OnIsAssetAcceptableForDrop.Execute(Object);
 		}
 		else
 		{
@@ -66,8 +117,7 @@ FReply SAssetDropTarget::OnDragOver( const FGeometry& MyGeometry, const FDragDro
 		bAllowDrop = false;
 	}
 
-	// Handle the reply if we are allowed to drop, otherwise do not handle it.
-	return bAllowDrop ? FReply::Handled() : FReply::Unhandled();
+	return bAllowDrop;
 }
 
 FReply SAssetDropTarget::OnDrop( const FGeometry& MyGeometry, const FDragDropEvent& DragDropEvent )
@@ -79,7 +129,7 @@ FReply SAssetDropTarget::OnDrop( const FGeometry& MyGeometry, const FDragDropEve
 	if( bAllowDrop )
 	{
 		bool bUnused;
-		UObject* Object = GetDroppedObject( DragDropEvent, bUnused );
+		UObject* Object = GetDroppedObject( DragDropEvent.GetOperation(), bUnused );
 
 		if( Object )
 		{
@@ -94,6 +144,7 @@ void SAssetDropTarget::OnDragEnter( const FGeometry& MyGeometry, const FDragDrop
 {
 	// initially we dont recognize this event
 	bIsDragEventRecognized = false;
+	bIsDragOver = true;
 }
 
 void SAssetDropTarget::OnDragLeave( const FDragDropEvent& DragDropEvent )
@@ -102,19 +153,20 @@ void SAssetDropTarget::OnDragLeave( const FDragDropEvent& DragDropEvent )
 	bIsDragEventRecognized = false;
 	// Disallow dropping if not dragged over.
 	bAllowDrop = false;
+
+	bIsDragOver = false;
 }
 
-UObject* SAssetDropTarget::GetDroppedObject( const FDragDropEvent& DragDropEvent, bool& bOutRecognizedEvent )
+UObject* SAssetDropTarget::GetDroppedObject(TSharedPtr<FDragDropOperation> DragDropOperation, bool& bOutRecognizedEvent) const
 {
 	bOutRecognizedEvent = false;
 	UObject* DroppedObject = NULL;
 
-	TSharedPtr<FDragDropOperation> Operation = DragDropEvent.GetOperation();
 	// Asset being dragged from content browser
-	if (Operation->IsOfType<FAssetDragDropOp>())
+	if ( DragDropOperation->IsOfType<FAssetDragDropOp>() )
 	{
 		bOutRecognizedEvent = true;
-		TSharedPtr<FAssetDragDropOp> DragDropOp = StaticCastSharedPtr<FAssetDragDropOp>(Operation);
+		TSharedPtr<FAssetDragDropOp> DragDropOp = StaticCastSharedPtr<FAssetDragDropOp>(DragDropOperation);
 
 		bool bCanDrop = DragDropOp->AssetData.Num() == 1;
 
@@ -127,9 +179,9 @@ UObject* SAssetDropTarget::GetDroppedObject( const FDragDropEvent& DragDropEvent
 		}
 	}
 	// Asset being dragged from some external source
-	else if (Operation->IsOfType<FExternalDragOperation>())
+	else if ( DragDropOperation->IsOfType<FExternalDragOperation>() )
 	{
-		TArray<FAssetData> DroppedAssetData = AssetUtil::ExtractAssetDataFromDrag(DragDropEvent);
+		TArray<FAssetData> DroppedAssetData = AssetUtil::ExtractAssetDataFromDrag(DragDropOperation);
 
 		if (DroppedAssetData.Num() == 1)
 		{
@@ -138,10 +190,10 @@ UObject* SAssetDropTarget::GetDroppedObject( const FDragDropEvent& DragDropEvent
 		}
 	}
 	// Actor being dragged?
-	else if (Operation->IsOfType<FActorDragDropOp>())
+	else if ( DragDropOperation->IsOfType<FActorDragDropOp>() )
 	{
 		bOutRecognizedEvent = true;
-		TSharedPtr<FActorDragDropOp> ActorDragDrop = StaticCastSharedPtr<FActorDragDropOp>(Operation);
+		TSharedPtr<FActorDragDropOp> ActorDragDrop = StaticCastSharedPtr<FActorDragDropOp>(DragDropOperation);
 
 		if (ActorDragDrop->Actors.Num() == 1)
 		{
@@ -151,3 +203,67 @@ UObject* SAssetDropTarget::GetDroppedObject( const FDragDropEvent& DragDropEvent
 
 	return DroppedObject;
 }
+
+int32 SAssetDropTarget::OnPaint(const FPaintArgs& Args, const FGeometry& AllottedGeometry, const FSlateRect& MyClippingRect, FSlateWindowElementList& OutDrawElements, int32 LayerId, const FWidgetStyle& InWidgetStyle, bool bParentEnabled) const
+{
+	LayerId = SCompoundWidget::OnPaint(Args, AllottedGeometry, MyClippingRect, OutDrawElements, LayerId, InWidgetStyle, bParentEnabled);
+
+	if ( GetDragOverlayVisibility().IsVisible() )
+	{
+		if ( bIsDragEventRecognized )
+		{
+			FLinearColor DashColor = bAllowDrop ? FLinearColor(0, 1, 0, 1) : FLinearColor(1, 0, 0, 1);
+
+			const FSlateBrush* HorizontalBrush = FEditorStyle::GetBrush("WideDash.Horizontal");
+			const FSlateBrush* VerticalBrush = FEditorStyle::GetBrush("WideDash.Vertical");
+
+			int32 DashLayer = LayerId + 1;
+
+			// Top
+			FSlateDrawElement::MakeBox(
+				OutDrawElements,
+				DashLayer,
+				AllottedGeometry.ToPaintGeometry(FVector2D(0, 0), FVector2D(AllottedGeometry.Size.X, HorizontalBrush->ImageSize.Y)),
+				HorizontalBrush,
+				MyClippingRect,
+				ESlateDrawEffect::None,
+				DashColor);
+
+			// Bottom
+			FSlateDrawElement::MakeBox(
+				OutDrawElements,
+				DashLayer,
+				AllottedGeometry.ToPaintGeometry(FVector2D(0, AllottedGeometry.Size.Y - HorizontalBrush->ImageSize.Y), FVector2D(AllottedGeometry.Size.X, HorizontalBrush->ImageSize.Y)),
+				HorizontalBrush,
+				MyClippingRect,
+				ESlateDrawEffect::None,
+				DashColor);
+
+			// Left
+			FSlateDrawElement::MakeBox(
+				OutDrawElements,
+				DashLayer,
+				AllottedGeometry.ToPaintGeometry(FVector2D(0, 0), FVector2D(VerticalBrush->ImageSize.X, AllottedGeometry.Size.Y)),
+				VerticalBrush,
+				MyClippingRect,
+				ESlateDrawEffect::None,
+				DashColor);
+
+			// Right
+			FSlateDrawElement::MakeBox(
+				OutDrawElements,
+				DashLayer,
+				AllottedGeometry.ToPaintGeometry(FVector2D(AllottedGeometry.Size.X - VerticalBrush->ImageSize.X, 0), FVector2D(VerticalBrush->ImageSize.X, AllottedGeometry.Size.Y)),
+				VerticalBrush,
+				MyClippingRect,
+				ESlateDrawEffect::None,
+				DashColor);
+
+			return DashLayer;
+		}
+	}
+
+	return LayerId;
+}
+
+#undef LOCTEXT_NAMESPACE
diff --git a/Engine/Source/Editor/EditorWidgets/Public/SAssetDropTarget.h b/Engine/Source/Editor/EditorWidgets/Public/SAssetDropTarget.h
--- a/Engine/Source/Editor/EditorWidgets/Public/SAssetDropTarget.h
+++ b/Engine/Source/Editor/EditorWidgets/Public/SAssetDropTarget.h
@@ -5,7 +5,7 @@
 /** 
  * A widget that displays a hover cue and handles dropping assets of allowed types onto this widget
  */
-class EDITORWIDGETS_API SAssetDropTarget : public SBorder
+class EDITORWIDGETS_API SAssetDropTarget : public SCompoundWidget
 {
 public:
 	/** Called when a valid asset is dropped */
@@ -14,8 +14,8 @@ public:
 	/** Called when we need to check if an asset type is valid for dropping */
 	DECLARE_DELEGATE_RetVal_OneParam( bool, FIsAssetAcceptableForDrop, const UObject* );
 
-	SLATE_BEGIN_ARGS( SAssetDropTarget )
-	{}
+	SLATE_BEGIN_ARGS(SAssetDropTarget)
+	{ }
 		/* Content to display for the in the drop target */
 		SLATE_DEFAULT_SLOT( FArguments, Content )
 		/** Called when a valid asset is dropped */
@@ -25,8 +25,8 @@ public:
 	SLATE_END_ARGS()
 
 	void Construct(const FArguments& InArgs );
+
 private:
-	const FSlateBrush* GetDragBorder() const;
 	FSlateColor GetDropBorderColor() const;
 
 	// SWidget interface
@@ -34,19 +34,34 @@ private:
 	virtual FReply OnDrop( const FGeometry& MyGeometry, const FDragDropEvent& DragDropEvent ) override;
 	virtual void OnDragEnter( const FGeometry& MyGeometry, const FDragDropEvent& DragDropEvent ) override;
 	virtual void OnDragLeave( const FDragDropEvent& DragDropEvent ) override;
+	virtual int32 OnPaint(const FPaintArgs& Args, const FGeometry& AllottedGeometry, const FSlateRect& MyClippingRect, FSlateWindowElementList& OutDrawElements, int32 LayerId, const FWidgetStyle& InWidgetStyle, bool bParentEnabled) const override;
 	// End of SWidget interface
 
+	/** @return Visibility of the overlay text when dragging is occuring. */
+	EVisibility GetDragOverlayVisibility() const;
+
+	/** @return The text to show in the overlay when a drag object is recognized and is acceptable or not. */
+	FText GetDragOverlayText() const;
+
+	/** Get the brightness on the background. */
+	FSlateColor GetBackgroundBrightness() const;
+
 	/**
 	 * Gets the dropped object from a drag drop event
 	 */
-	UObject* GetDroppedObject( const FDragDropEvent& DragDropEvent, bool& bOutRecognizedEvent );
+	UObject* GetDroppedObject(TSharedPtr<FDragDropOperation> DragDropOperation, bool& bOutRecognizedEvent) const;
+
+	/** Utility function that can determine if this asset drop target accepts the item being dragged and dropped. */
+	bool AllowDrop(TSharedPtr<FDragDropOperation> DragDropOperation) const;
 private:
 	/** Delegate to call when an asset is dropped */
 	FOnAssetDropped OnAssetDropped;
 	/** Delegate to call to check validity of the asset */
 	FIsAssetAcceptableForDrop OnIsAssetAcceptableForDrop;
 	/** Whether or not we are being dragged over by a recognized event*/
-	bool bIsDragEventRecognized;
+	mutable bool bIsDragEventRecognized;
 	/** Whether or not we currently allow dropping */
-	bool bAllowDrop;
+	mutable bool bAllowDrop;
+	/** Is the drag operation currently over our airspace? */
+	mutable bool bIsDragOver;
 };
\ No newline at end of file
