[INTEGRATE] Change 2426450 by [REDACTED_PERSON][REDACTED_HANDLE] on 2015/01/30 17:18:55

	Editor - Tweaking the drop target some.

diff --git a/Engine/Source/Editor/EditorWidgets/Private/SAssetDropTarget.cpp b/Engine/Source/Editor/EditorWidgets/Private/SAssetDropTarget.cpp
--- a/Engine/Source/Editor/EditorWidgets/Private/SAssetDropTarget.cpp
+++ b/Engine/Source/Editor/EditorWidgets/Private/SAssetDropTarget.cpp
@@ -14,147 +14,55 @@ void SAssetDropTarget::Construct(const FArguments& InArgs )
 	OnAssetDropped = InArgs._OnAssetDropped;
 	OnIsAssetAcceptableForDrop = InArgs._OnIsAssetAcceptableForDrop;
 
-	bIsDragEventRecognized = false;
-	bAllowDrop = false;
-	bIsDragOver = false;
-
-	ChildSlot
-	[
-		SNew(SOverlay)
-			
-		+ SOverlay::Slot()
+	SDropTarget::Construct(
+		SDropTarget::FArguments()
+		.OnDrop(this, &SAssetDropTarget::OnDropped)
 		[
 			InArgs._Content.Widget
-		]
-
-		+ SOverlay::Slot()
-		[
-			SNew(SBorder)
-			.Visibility(this, &SAssetDropTarget::GetDragOverlayVisibility)
-			.BorderImage(FEditorStyle::GetBrush("NoBorder"))
-			.BorderBackgroundColor(this, &SAssetDropTarget::GetDropBorderColor)
-			[
-				SNew(SBorder)
-				.BorderImage(FEditorStyle::GetBrush("WhiteBrush"))
-				.BorderBackgroundColor(this, &SAssetDropTarget::GetBackgroundBrightness)
-				.Padding(15.0f)
-				[
-					SNew(SScaleBox)
-					.Stretch(EStretch::ScaleToFit)
-					[
-						SNew(STextBlock)
-						.Text(this, &SAssetDropTarget::GetDragOverlayText)
-						.ColorAndOpacity(this, &SAssetDropTarget::GetDropBorderColor)
-					]
-				]
-			]
-		]
-	];
+		]);
 }
 
-FSlateColor SAssetDropTarget::GetBackgroundBrightness() const
+FReply SAssetDropTarget::OnDropped(TSharedPtr<FDragDropOperation> DragDropOperation)
 {
-	return ( bAllowDrop && bIsDragOver ) ? FLinearColor(1, 1, 1, 0.40f) : FLinearColor(1, 1, 1, 0.25f);
-}
+	bool bUnused;
+	UObject* Object = GetDroppedObject(DragDropOperation, bUnused);
 
-EVisibility SAssetDropTarget::GetDragOverlayVisibility() const
-{
-	if ( FSlateApplication::Get().IsDragDropping() )
+	if ( Object )
 	{
-		if ( AllowDrop(FSlateApplication::Get().GetDragDroppingContent()) || bIsDragOver )
-		{
-			return EVisibility::HitTestInvisible;
-		}
+		OnAssetDropped.ExecuteIfBound(Object);
 	}
 
-	return EVisibility::Hidden;
-}
-
-FText SAssetDropTarget::GetDragOverlayText() const
-{
-	return bAllowDrop ? FText::GetEmpty() : LOCTEXT("Nope", "Nope");
-}
-
-FSlateColor SAssetDropTarget::GetDropBorderColor() const
-{
-	return bIsDragEventRecognized ? 
-		bAllowDrop ? FLinearColor(0,1,0,1) : FLinearColor(1,0,0,1)
-		: FLinearColor::White;
-}
-
-FReply SAssetDropTarget::OnDragOver( const FGeometry& MyGeometry, const FDragDropEvent& DragDropEvent )
-{
-	// Handle the reply if we are allowed to drop, otherwise do not handle it.
-	return AllowDrop(DragDropEvent.GetOperation()) ? FReply::Handled() : FReply::Unhandled();
+	return FReply::Handled();
 }
 
-bool SAssetDropTarget::AllowDrop(TSharedPtr<FDragDropOperation> DragDropOperation) const
+bool SAssetDropTarget::OnAllowDrop(TSharedPtr<FDragDropOperation> DragDropOperation) const
 {
-	bool bRecognizedEvent = false;
-	UObject* Object = GetDroppedObject(DragDropOperation, bRecognizedEvent);
-
-	// Not being dragged over by a recognizable event
-	bIsDragEventRecognized = bRecognizedEvent;
-
-	bAllowDrop = false;
+	bool bUnused = false;
+	UObject* Object = GetDroppedObject(DragDropOperation, bUnused);
 
 	if ( Object )
 	{
 		// Check and see if its valid to drop this object
 		if ( OnIsAssetAcceptableForDrop.IsBound() )
 		{
-			bAllowDrop = OnIsAssetAcceptableForDrop.Execute(Object);
+			return OnIsAssetAcceptableForDrop.Execute(Object);
 		}
 		else
 		{
 			// If no delegate is bound assume its always valid to drop this object
-			bAllowDrop = true;
-		}
-	}
-	else
-	{
-		// No object so we dont allow dropping
-		bAllowDrop = false;
-	}
-
-	return bAllowDrop;
-}
-
-FReply SAssetDropTarget::OnDrop( const FGeometry& MyGeometry, const FDragDropEvent& DragDropEvent )
-{
-	// We've dropped an asset so we are no longer being dragged over
-	bIsDragEventRecognized = false;
-
-	// if we allow drop, call a delegate to handle the drop
-	if( bAllowDrop )
-	{
-		bool bUnused;
-		UObject* Object = GetDroppedObject( DragDropEvent.GetOperation(), bUnused );
-
-		if( Object )
-		{
-			OnAssetDropped.ExecuteIfBound( Object );
+			return true;
 		}
 	}
 
-	return FReply::Handled();
-}
-
-void SAssetDropTarget::OnDragEnter( const FGeometry& MyGeometry, const FDragDropEvent& DragDropEvent )
-{
-	// initially we dont recognize this event
-	bIsDragEventRecognized = false;
-	bIsDragOver = true;
+	return false;
 }
 
-void SAssetDropTarget::OnDragLeave( const FDragDropEvent& DragDropEvent )
+bool SAssetDropTarget::OnIsRecognized(TSharedPtr<FDragDropOperation> DragDropOperation) const
 {
-	// No longer being dragged over
-	bIsDragEventRecognized = false;
-	// Disallow dropping if not dragged over.
-	bAllowDrop = false;
+	bool bRecognizedEvent = false;
+	UObject* Object = GetDroppedObject(DragDropOperation, bRecognizedEvent);
 
-	bIsDragOver = false;
+	return bRecognizedEvent;
 }
 
 UObject* SAssetDropTarget::GetDroppedObject(TSharedPtr<FDragDropOperation> DragDropOperation, bool& bOutRecognizedEvent) const
@@ -204,66 +112,4 @@ UObject* SAssetDropTarget::GetDroppedObject(TSharedPtr<FDragDropOperation> DragD
 	return DroppedObject;
 }
 
-int32 SAssetDropTarget::OnPaint(const FPaintArgs& Args, const FGeometry& AllottedGeometry, const FSlateRect& MyClippingRect, FSlateWindowElementList& OutDrawElements, int32 LayerId, const FWidgetStyle& InWidgetStyle, bool bParentEnabled) const
-{
-	LayerId = SCompoundWidget::OnPaint(Args, AllottedGeometry, MyClippingRect, OutDrawElements, LayerId, InWidgetStyle, bParentEnabled);
-
-	if ( GetDragOverlayVisibility().IsVisible() )
-	{
-		if ( bIsDragEventRecognized )
-		{
-			FLinearColor DashColor = bAllowDrop ? FLinearColor(0, 1, 0, 1) : FLinearColor(1, 0, 0, 1);
-
-			const FSlateBrush* HorizontalBrush = FEditorStyle::GetBrush("WideDash.Horizontal");
-			const FSlateBrush* VerticalBrush = FEditorStyle::GetBrush("WideDash.Vertical");
-
-			int32 DashLayer = LayerId + 1;
-
-			// Top
-			FSlateDrawElement::MakeBox(
-				OutDrawElements,
-				DashLayer,
-				AllottedGeometry.ToPaintGeometry(FVector2D(0, 0), FVector2D(AllottedGeometry.Size.X, HorizontalBrush->ImageSize.Y)),
-				HorizontalBrush,
-				MyClippingRect,
-				ESlateDrawEffect::None,
-				DashColor);
-
-			// Bottom
-			FSlateDrawElement::MakeBox(
-				OutDrawElements,
-				DashLayer,
-				AllottedGeometry.ToPaintGeometry(FVector2D(0, AllottedGeometry.Size.Y - HorizontalBrush->ImageSize.Y), FVector2D(AllottedGeometry.Size.X, HorizontalBrush->ImageSize.Y)),
-				HorizontalBrush,
-				MyClippingRect,
-				ESlateDrawEffect::None,
-				DashColor);
-
-			// Left
-			FSlateDrawElement::MakeBox(
-				OutDrawElements,
-				DashLayer,
-				AllottedGeometry.ToPaintGeometry(FVector2D(0, 0), FVector2D(VerticalBrush->ImageSize.X, AllottedGeometry.Size.Y)),
-				VerticalBrush,
-				MyClippingRect,
-				ESlateDrawEffect::None,
-				DashColor);
-
-			// Right
-			FSlateDrawElement::MakeBox(
-				OutDrawElements,
-				DashLayer,
-				AllottedGeometry.ToPaintGeometry(FVector2D(AllottedGeometry.Size.X - VerticalBrush->ImageSize.X, 0), FVector2D(VerticalBrush->ImageSize.X, AllottedGeometry.Size.Y)),
-				VerticalBrush,
-				MyClippingRect,
-				ESlateDrawEffect::None,
-				DashColor);
-
-			return DashLayer;
-		}
-	}
-
-	return LayerId;
-}
-
 #undef LOCTEXT_NAMESPACE
diff --git a/Engine/Source/Editor/EditorWidgets/Private/SDropTarget.cpp b/Engine/Source/Editor/EditorWidgets/Private/SDropTarget.cpp
new file mode 100644
--- /dev/null
+++ b/Engine/Source/Editor/EditorWidgets/Private/SDropTarget.cpp
@@ -0,0 +1,205 @@
+// Copyright 1998-2015 Epic Games, Inc. All Rights Reserved.
+
+#include "EditorWidgetsPrivatePCH.h"
+
+#include "SDropTarget.h"
+#include "AssetDragDropOp.h"
+#include "ActorDragDropOp.h"
+#include "AssetSelection.h"
+#include "SScaleBox.h"
+
+#define LOCTEXT_NAMESPACE "EditorWidgets"
+
+void SDropTarget::Construct(const FArguments& InArgs)
+{
+	DroppedEvent = InArgs._OnDrop;
+	AllowDropEvent = InArgs._OnAllowDrop;
+	IsRecognizedEvent = InArgs._OnIsRecognized;
+
+	bIsDragEventRecognized = false;
+	bAllowDrop = false;
+	bIsDragOver = false;
+
+	ChildSlot
+	[
+		SNew(SOverlay)
+			
+		+ SOverlay::Slot()
+		[
+			InArgs._Content.Widget
+		]
+
+		+ SOverlay::Slot()
+		[
+			SNew(SBorder)
+			.Visibility(this, &SDropTarget::GetDragOverlayVisibility)
+			.BorderImage(FEditorStyle::GetBrush("NoBorder"))
+			.BorderBackgroundColor(this, &SDropTarget::GetDropBorderColor)
+			[
+				SNew(SBorder)
+				.BorderImage(FEditorStyle::GetBrush("WhiteBrush"))
+				.BorderBackgroundColor(this, &SDropTarget::GetBackgroundBrightness)
+				.Padding(15.0f)
+			]
+		]
+	];
+}
+
+FSlateColor SDropTarget::GetBackgroundBrightness() const
+{
+	return ( bIsDragOver ) ? FLinearColor(1, 1, 1, 0.50f) : FLinearColor(1, 1, 1, 0.25f);
+}
+
+EVisibility SDropTarget::GetDragOverlayVisibility() const
+{
+	if ( FSlateApplication::Get().IsDragDropping() )
+	{
+		if ( AllowDrop(FSlateApplication::Get().GetDragDroppingContent()) || (bIsDragOver && bIsDragEventRecognized) )
+		{
+			return EVisibility::HitTestInvisible;
+		}
+	}
+
+	return EVisibility::Hidden;
+}
+
+FSlateColor SDropTarget::GetDropBorderColor() const
+{
+	return bIsDragEventRecognized ? 
+		bAllowDrop ? FLinearColor(0,1,0,1) : FLinearColor(1,0,0,1)
+		: FLinearColor::White;
+}
+
+FReply SDropTarget::OnDragOver(const FGeometry& MyGeometry, const FDragDropEvent& DragDropEvent)
+{
+	// Handle the reply if we are allowed to drop, otherwise do not handle it.
+	return AllowDrop(DragDropEvent.GetOperation()) ? FReply::Handled() : FReply::Unhandled();
+}
+
+bool SDropTarget::AllowDrop(TSharedPtr<FDragDropOperation> DragDropOperation) const
+{
+	bAllowDrop = OnAllowDrop(DragDropOperation);
+	bIsDragEventRecognized = OnIsRecognized(DragDropOperation) || bAllowDrop;
+
+	return bAllowDrop;
+}
+
+bool SDropTarget::OnAllowDrop(TSharedPtr<FDragDropOperation> DragDropOperation) const
+{
+	if ( AllowDropEvent.IsBound() )
+	{
+		return AllowDropEvent.Execute(DragDropOperation);
+	}
+
+	return false;
+}
+
+bool SDropTarget::OnIsRecognized(TSharedPtr<FDragDropOperation> DragDropOperation) const
+{
+	if ( IsRecognizedEvent.IsBound() )
+	{
+		return IsRecognizedEvent.Execute(DragDropOperation);
+	}
+
+	return false;
+}
+
+FReply SDropTarget::OnDrop(const FGeometry& MyGeometry, const FDragDropEvent& DragDropEvent)
+{
+	const bool bCurrentbAllowDrop = bAllowDrop;
+
+	// We've dropped an asset so we are no longer being dragged over
+	bIsDragEventRecognized = false;
+	bIsDragOver = false;
+	bAllowDrop = false;
+
+	// if we allow drop, call a delegate to handle the drop
+	if ( bCurrentbAllowDrop )
+	{
+		if ( DroppedEvent.IsBound() )
+		{
+			return DroppedEvent.Execute(DragDropEvent.GetOperation());
+		}
+	}
+
+	return FReply::Handled();
+}
+
+void SDropTarget::OnDragEnter(const FGeometry& MyGeometry, const FDragDropEvent& DragDropEvent)
+{
+	// initially we dont recognize this event
+	bIsDragEventRecognized = false;
+	bIsDragOver = true;
+}
+
+void SDropTarget::OnDragLeave(const FDragDropEvent& DragDropEvent)
+{
+	// No longer being dragged over
+	bIsDragEventRecognized = false;
+	// Disallow dropping if not dragged over.
+	bAllowDrop = false;
+
+	bIsDragOver = false;
+}
+
+int32 SDropTarget::OnPaint(const FPaintArgs& Args, const FGeometry& AllottedGeometry, const FSlateRect& MyClippingRect, FSlateWindowElementList& OutDrawElements, int32 LayerId, const FWidgetStyle& InWidgetStyle, bool bParentEnabled) const
+{
+	LayerId = SCompoundWidget::OnPaint(Args, AllottedGeometry, MyClippingRect, OutDrawElements, LayerId, InWidgetStyle, bParentEnabled);
+
+	if ( GetDragOverlayVisibility().IsVisible() )
+	{
+		if ( bIsDragEventRecognized )
+		{
+			FLinearColor DashColor = bAllowDrop ? FLinearColor(0, 1, 0, 1) : FLinearColor(1, 0, 0, 1);
+
+			const FSlateBrush* HorizontalBrush = FEditorStyle::GetBrush("WideDash.Horizontal");
+			const FSlateBrush* VerticalBrush = FEditorStyle::GetBrush("WideDash.Vertical");
+
+			int32 DashLayer = LayerId + 1;
+
+			// Top
+			FSlateDrawElement::MakeBox(
+				OutDrawElements,
+				DashLayer,
+				AllottedGeometry.ToPaintGeometry(FVector2D(0, 0), FVector2D(AllottedGeometry.Size.X, HorizontalBrush->ImageSize.Y)),
+				HorizontalBrush,
+				MyClippingRect,
+				ESlateDrawEffect::None,
+				DashColor);
+
+			// Bottom
+			FSlateDrawElement::MakeBox(
+				OutDrawElements,
+				DashLayer,
+				AllottedGeometry.ToPaintGeometry(FVector2D(0, AllottedGeometry.Size.Y - HorizontalBrush->ImageSize.Y), FVector2D(AllottedGeometry.Size.X, HorizontalBrush->ImageSize.Y)),
+				HorizontalBrush,
+				MyClippingRect,
+				ESlateDrawEffect::None,
+				DashColor);
+
+			// Left
+			FSlateDrawElement::MakeBox(
+				OutDrawElements,
+				DashLayer,
+				AllottedGeometry.ToPaintGeometry(FVector2D(0, 0), FVector2D(VerticalBrush->ImageSize.X, AllottedGeometry.Size.Y)),
+				VerticalBrush,
+				MyClippingRect,
+				ESlateDrawEffect::None,
+				DashColor);
+
+			// Right
+			FSlateDrawElement::MakeBox(
+				OutDrawElements,
+				DashLayer,
+				AllottedGeometry.ToPaintGeometry(FVector2D(AllottedGeometry.Size.X - VerticalBrush->ImageSize.X, 0), FVector2D(VerticalBrush->ImageSize.X, AllottedGeometry.Size.Y)),
+				VerticalBrush,
+				MyClippingRect,
+				ESlateDrawEffect::None,
+				DashColor);
+
+			return DashLayer;
+		}
+	}
+
+	return LayerId;
+}
diff --git a/Engine/Source/Editor/EditorWidgets/Public/SAssetDropTarget.h b/Engine/Source/Editor/EditorWidgets/Public/SAssetDropTarget.h
--- a/Engine/Source/Editor/EditorWidgets/Public/SAssetDropTarget.h
+++ b/Engine/Source/Editor/EditorWidgets/Public/SAssetDropTarget.h
@@ -2,10 +2,12 @@
 
 #pragma once
 
-/** 
+#include "SDropTarget.h"
+
+/**
  * A widget that displays a hover cue and handles dropping assets of allowed types onto this widget
  */
-class EDITORWIDGETS_API SAssetDropTarget : public SCompoundWidget
+class EDITORWIDGETS_API SAssetDropTarget : public SDropTarget
 {
 public:
 	/** Called when a valid asset is dropped */
@@ -26,42 +28,17 @@ public:
 
 	void Construct(const FArguments& InArgs );
 
-private:
-	FSlateColor GetDropBorderColor() const;
-
-	// SWidget interface
-	virtual FReply OnDragOver( const FGeometry& MyGeometry, const FDragDropEvent& DragDropEvent ) override;
-	virtual FReply OnDrop( const FGeometry& MyGeometry, const FDragDropEvent& DragDropEvent ) override;
-	virtual void OnDragEnter( const FGeometry& MyGeometry, const FDragDropEvent& DragDropEvent ) override;
-	virtual void OnDragLeave( const FDragDropEvent& DragDropEvent ) override;
-	virtual int32 OnPaint(const FPaintArgs& Args, const FGeometry& AllottedGeometry, const FSlateRect& MyClippingRect, FSlateWindowElementList& OutDrawElements, int32 LayerId, const FWidgetStyle& InWidgetStyle, bool bParentEnabled) const override;
-	// End of SWidget interface
-
-	/** @return Visibility of the overlay text when dragging is occuring. */
-	EVisibility GetDragOverlayVisibility() const;
+protected:
+	FReply OnDropped(TSharedPtr<FDragDropOperation> DragDropOperation);
+	virtual bool OnAllowDrop(TSharedPtr<FDragDropOperation> DragDropOperation) const override;
+	virtual bool OnIsRecognized(TSharedPtr<FDragDropOperation> DragDropOperation) const override;
 
-	/** @return The text to show in the overlay when a drag object is recognized and is acceptable or not. */
-	FText GetDragOverlayText() const;
-
-	/** Get the brightness on the background. */
-	FSlateColor GetBackgroundBrightness() const;
-
-	/**
-	 * Gets the dropped object from a drag drop event
-	 */
+private:
 	UObject* GetDroppedObject(TSharedPtr<FDragDropOperation> DragDropOperation, bool& bOutRecognizedEvent) const;
 
-	/** Utility function that can determine if this asset drop target accepts the item being dragged and dropped. */
-	bool AllowDrop(TSharedPtr<FDragDropOperation> DragDropOperation) const;
 private:
 	/** Delegate to call when an asset is dropped */
 	FOnAssetDropped OnAssetDropped;
 	/** Delegate to call to check validity of the asset */
 	FIsAssetAcceptableForDrop OnIsAssetAcceptableForDrop;
-	/** Whether or not we are being dragged over by a recognized event*/
-	mutable bool bIsDragEventRecognized;
-	/** Whether or not we currently allow dropping */
-	mutable bool bAllowDrop;
-	/** Is the drag operation currently over our airspace? */
-	mutable bool bIsDragOver;
 };
\ No newline at end of file
diff --git a/Engine/Source/Editor/EditorWidgets/Public/SDropTarget.h b/Engine/Source/Editor/EditorWidgets/Public/SDropTarget.h
new file mode 100644
--- /dev/null
+++ b/Engine/Source/Editor/EditorWidgets/Public/SDropTarget.h
@@ -0,0 +1,68 @@
+// Copyright 1998-2015 Epic Games, Inc. All Rights Reserved.
+
+#pragma once
+
+/** 
+ * A widget that displays a hover cue and handles dropping assets of allowed types onto this widget
+ */
+class EDITORWIDGETS_API SDropTarget : public SCompoundWidget
+{
+public:
+	/** Called when a valid asset is dropped */
+	DECLARE_DELEGATE_RetVal_OneParam(FReply, FOnDrop, TSharedPtr<FDragDropOperation>);
+
+	DECLARE_DELEGATE_RetVal_OneParam(bool, FVerifyDrag, TSharedPtr<FDragDropOperation>);
+
+	SLATE_BEGIN_ARGS(SDropTarget)
+	{ }
+		/* Content to display for the in the drop target */
+		SLATE_DEFAULT_SLOT( FArguments, Content )
+		/** Called when a valid asset is dropped */
+		SLATE_EVENT(FOnDrop, OnDrop)
+		/** Called to check if an asset is acceptable for dropping */
+		SLATE_EVENT(FVerifyDrag, OnAllowDrop)
+		/** Called to check if an asset is acceptable for dropping */
+		SLATE_EVENT(FVerifyDrag, OnIsRecognized)
+	SLATE_END_ARGS()
+
+	void Construct(const FArguments& InArgs );
+
+protected:
+
+	bool AllowDrop(TSharedPtr<FDragDropOperation> DragDropOperation) const;
+
+	virtual bool OnAllowDrop(TSharedPtr<FDragDropOperation> DragDropOperation) const;
+	virtual bool OnIsRecognized(TSharedPtr<FDragDropOperation> DragDropOperation) const;
+
+protected:
+	FSlateColor GetDropBorderColor() const;
+
+	// SWidget interface
+	virtual FReply OnDragOver( const FGeometry& MyGeometry, const FDragDropEvent& DragDropEvent ) override;
+	virtual FReply OnDrop( const FGeometry& MyGeometry, const FDragDropEvent& DragDropEvent ) override;
+	virtual void OnDragEnter( const FGeometry& MyGeometry, const FDragDropEvent& DragDropEvent ) override;
+	virtual void OnDragLeave( const FDragDropEvent& DragDropEvent ) override;
+	virtual int32 OnPaint(const FPaintArgs& Args, const FGeometry& AllottedGeometry, const FSlateRect& MyClippingRect, FSlateWindowElementList& OutDrawElements, int32 LayerId, const FWidgetStyle& InWidgetStyle, bool bParentEnabled) const override;
+	// End of SWidget interface
+
+	/** @return Visibility of the overlay text when dragging is occuring. */
+	EVisibility GetDragOverlayVisibility() const;
+
+	/** Get the brightness on the background. */
+	FSlateColor GetBackgroundBrightness() const;
+
+private:
+	/** Delegate to call when an asset is dropped */
+	FOnDrop DroppedEvent;
+	/** Delegate to call to check validity of the asset */
+	FVerifyDrag AllowDropEvent;
+	/** Delegate to call to check validity of the asset */
+	FVerifyDrag IsRecognizedEvent;
+
+	/** Whether or not we are being dragged over by a recognized event*/
+	mutable bool bIsDragEventRecognized;
+	/** Whether or not we currently allow dropping */
+	mutable bool bAllowDrop;
+	/** Is the drag operation currently over our airspace? */
+	mutable bool bIsDragOver;
+};
\ No newline at end of file
