[INTEGRATE] Change 2431888 by [REDACTED_PERSON][REDACTED_HANDLE] on 2015/02/04 12:03:13

	Editor - Improving the SDropTarget Editor Widget to be more customizable.

diff --git a/Engine/Source/Editor/EditorWidgets/Private/SDropTarget.cpp b/Engine/Source/Editor/EditorWidgets/Private/SDropTarget.cpp
--- a/Engine/Source/Editor/EditorWidgets/Private/SDropTarget.cpp
+++ b/Engine/Source/Editor/EditorWidgets/Private/SDropTarget.cpp
@@ -20,6 +20,15 @@ void SDropTarget::Construct(const FArguments& InArgs)
 	bAllowDrop = false;
 	bIsDragOver = false;
 
+	ValidColor = InArgs._ValidColor;
+	InvalidColor = InArgs._InvalidColor;
+
+	BackgroundColor = InArgs._BackgroundColor;
+	BackgroundColorHover = InArgs._BackgroundColorHover;
+
+	VerticalImage = InArgs._VerticalImage;
+	HorizontalImage = InArgs._HorizontalImage;
+
 	ChildSlot
 	[
 		SNew(SOverlay)
@@ -31,15 +40,12 @@ void SDropTarget::Construct(const FArguments& InArgs)
 
 		+ SOverlay::Slot()
 		[
-			SNew(SBorder)
+			SNew(SBox)
 			.Visibility(this, &SDropTarget::GetDragOverlayVisibility)
-			.BorderImage(FEditorStyle::GetBrush("NoBorder"))
-			.BorderBackgroundColor(this, &SDropTarget::GetDropBorderColor)
 			[
 				SNew(SBorder)
-				.BorderImage(FEditorStyle::GetBrush("WhiteBrush"))
+				.BorderImage(InArgs._BackgroundImage)
 				.BorderBackgroundColor(this, &SDropTarget::GetBackgroundBrightness)
-				.Padding(15.0f)
 			]
 		]
 	];
@@ -47,7 +53,7 @@ void SDropTarget::Construct(const FArguments& InArgs)
 
 FSlateColor SDropTarget::GetBackgroundBrightness() const
 {
-	return ( bIsDragOver ) ? FLinearColor(1, 1, 1, 0.50f) : FLinearColor(1, 1, 1, 0.25f);
+	return ( bIsDragOver ) ? BackgroundColorHover : BackgroundColor;
 }
 
 EVisibility SDropTarget::GetDragOverlayVisibility() const
@@ -63,13 +69,6 @@ EVisibility SDropTarget::GetDragOverlayVisibility() const
 	return EVisibility::Hidden;
 }
 
-FSlateColor SDropTarget::GetDropBorderColor() const
-{
-	return bIsDragEventRecognized ? 
-		bAllowDrop ? FLinearColor(0,1,0,1) : FLinearColor(1,0,0,1)
-		: FLinearColor::White;
-}
-
 FReply SDropTarget::OnDragOver(const FGeometry& MyGeometry, const FDragDropEvent& DragDropEvent)
 {
 	// Handle the reply if we are allowed to drop, otherwise do not handle it.
@@ -150,7 +149,7 @@ int32 SDropTarget::OnPaint(const FPaintArgs& Args, const FGeometry& AllottedGeom
 	{
 		if ( bIsDragEventRecognized )
 		{
-			FLinearColor DashColor = bAllowDrop ? FLinearColor(0, 1, 0, 1) : FLinearColor(1, 0, 0, 1);
+			FLinearColor DashColor = bAllowDrop ? ValidColor : InvalidColor;
 
 			const FSlateBrush* HorizontalBrush = FEditorStyle::GetBrush("WideDash.Horizontal");
 			const FSlateBrush* VerticalBrush = FEditorStyle::GetBrush("WideDash.Vertical");
diff --git a/Engine/Source/Editor/EditorWidgets/Public/SDropTarget.h b/Engine/Source/Editor/EditorWidgets/Public/SDropTarget.h
--- a/Engine/Source/Editor/EditorWidgets/Public/SDropTarget.h
+++ b/Engine/Source/Editor/EditorWidgets/Public/SDropTarget.h
@@ -14,9 +14,30 @@ public:
 	DECLARE_DELEGATE_RetVal_OneParam(bool, FVerifyDrag, TSharedPtr<FDragDropOperation>);
 
 	SLATE_BEGIN_ARGS(SDropTarget)
+		: _ValidColor(FLinearColor(0, 1, 0, 1))
+		, _InvalidColor(FLinearColor(1, 0, 0, 1))
+		, _BackgroundColor(FLinearColor(1, 1, 1, 0.50f))
+		, _BackgroundColorHover(FLinearColor(1, 1, 1, 0.25f))
+		, _VerticalImage(FEditorStyle::GetBrush("WideDash.Vertical"))
+		, _HorizontalImage(FEditorStyle::GetBrush("WideDash.Horizontal"))
+		, _BackgroundImage(FEditorStyle::GetBrush("WhiteBrush"))
 	{ }
 		/* Content to display for the in the drop target */
 		SLATE_DEFAULT_SLOT( FArguments, Content )
+		/** The color of the vertical/horizontal images when the drop data is valid */
+		SLATE_ARGUMENT(FLinearColor, ValidColor)
+		/** The color of the vertical/horizontal images when the drop data is not valid */
+		SLATE_ARGUMENT(FLinearColor, InvalidColor)
+		/** The color multiplier that is applied to the background image. */
+		SLATE_ARGUMENT(FLinearColor, BackgroundColor)
+		/** The color multiplier that is applied to the background image on hover. */
+		SLATE_ARGUMENT(FLinearColor, BackgroundColorHover)
+		/** Vertical border image that is used. */
+		SLATE_ARGUMENT(const FSlateBrush*, VerticalImage)
+		/** Horizontal border image that is used. */
+		SLATE_ARGUMENT(const FSlateBrush*, HorizontalImage)
+		/** The background image that is applied after the surface. */
+		SLATE_ARGUMENT(const FSlateBrush*, BackgroundImage)
 		/** Called when a valid asset is dropped */
 		SLATE_EVENT(FOnDrop, OnDrop)
 		/** Called to check if an asset is acceptable for dropping */
@@ -35,8 +56,6 @@ protected:
 	virtual bool OnIsRecognized(TSharedPtr<FDragDropOperation> DragDropOperation) const;
 
 protected:
-	FSlateColor GetDropBorderColor() const;
-
 	// SWidget interface
 	virtual FReply OnDragOver( const FGeometry& MyGeometry, const FDragDropEvent& DragDropEvent ) override;
 	virtual FReply OnDrop( const FGeometry& MyGeometry, const FDragDropEvent& DragDropEvent ) override;
@@ -45,7 +64,7 @@ protected:
 	virtual int32 OnPaint(const FPaintArgs& Args, const FGeometry& AllottedGeometry, const FSlateRect& MyClippingRect, FSlateWindowElementList& OutDrawElements, int32 LayerId, const FWidgetStyle& InWidgetStyle, bool bParentEnabled) const override;
 	// End of SWidget interface
 
-	/** @return Visibility of the overlay text when dragging is occuring. */
+	/** @return Visibility of the overlay text when dragging is occurring. */
 	EVisibility GetDragOverlayVisibility() const;
 
 	/** Get the brightness on the background. */
@@ -59,6 +78,21 @@ private:
 	/** Delegate to call to check validity of the asset */
 	FVerifyDrag IsRecognizedEvent;
 
+	/** The color of the vertical/horizontal images when the drop data is valid */
+	FLinearColor ValidColor;
+	/** The color of the vertical/horizontal images when the drop data is not valid */
+	FLinearColor InvalidColor;
+	/** The color multiplier that is applied to the background image. */
+	FLinearColor BackgroundColor;
+	/** The color multiplier that is applied to the background image on hover. */
+	FLinearColor BackgroundColorHover;
+	/** The background image that is applied after the surface. */
+	const FSlateBrush* BackgroundImage;
+	/** Vertical border image that is used. */
+	const FSlateBrush* VerticalImage;
+	/** Horizontal border image that is used. */
+	const FSlateBrush* HorizontalImage;
+
 	/** Whether or not we are being dragged over by a recognized event*/
 	mutable bool bIsDragEventRecognized;
 	/** Whether or not we currently allow dropping */
