Copying //UE4/Dev-Editor to //UE4/Dev-Main (Source: //UE4/Dev-Editor @ 3279756)

diff --git a/Engine/Source/Editor/EditorWidgets/Private/SDropTarget.cpp b/Engine/Source/Editor/EditorWidgets/Private/SDropTarget.cpp
--- a/Engine/Source/Editor/EditorWidgets/Private/SDropTarget.cpp
+++ b/Engine/Source/Editor/EditorWidgets/Private/SDropTarget.cpp
@@ -119,9 +119,11 @@ FReply SDropTarget::OnDrop(const FGeometry& MyGeometry, const FDragDropEvent& Dr
 		{
 			return DroppedEvent.Execute(DragDropEvent.GetOperation());
 		}
+
+		return FReply::Handled();
 	}
 
-	return FReply::Handled();
+	return FReply::Unhandled();
 }
 
 void SDropTarget::OnDragEnter(const FGeometry& MyGeometry, const FDragDropEvent& DragDropEvent)
diff --git a/Engine/Source/Editor/EditorWidgets/Private/SObjectNameEditableTextBox.cpp b/Engine/Source/Editor/EditorWidgets/Private/SObjectNameEditableTextBox.cpp
--- a/Engine/Source/Editor/EditorWidgets/Private/SObjectNameEditableTextBox.cpp
+++ b/Engine/Source/Editor/EditorWidgets/Private/SObjectNameEditableTextBox.cpp
@@ -79,7 +79,7 @@ int32 SObjectNameEditableTextBox::OnPaint( const FPaintArgs& Args, const FGeomet
 
 	// See if a disabled effect should be used
 	bool bEnabled = ShouldBeEnabled( bParentEnabled );
-	ESlateDrawEffect::Type DrawEffects = (bEnabled) ? ESlateDrawEffect::None : ESlateDrawEffect::DisabledEffect;
+	ESlateDrawEffect DrawEffects = (bEnabled) ? ESlateDrawEffect::None : ESlateDrawEffect::DisabledEffect;
 
 	const double CurrentTime = FSlateApplication::Get().GetCurrentTime();
 
diff --git a/Engine/Source/Editor/EditorWidgets/Private/STextPropertyEditableTextBox.cpp b/Engine/Source/Editor/EditorWidgets/Private/STextPropertyEditableTextBox.cpp
--- a/Engine/Source/Editor/EditorWidgets/Private/STextPropertyEditableTextBox.cpp
+++ b/Engine/Source/Editor/EditorWidgets/Private/STextPropertyEditableTextBox.cpp
@@ -8,12 +8,19 @@
 #include "Widgets/Layout/SBox.h"
 #include "Widgets/Layout/SGridPanel.h"
 #include "Widgets/Layout/SUniformGridPanel.h"
+#include "Widgets/Input/SButton.h"
 #include "Widgets/Input/SMultiLineEditableTextBox.h"
 #include "Widgets/Input/SEditableTextBox.h"
 #include "Widgets/Input/SComboButton.h"
+#include "Widgets/Input/SComboBox.h"
 #include "Widgets/Input/SCheckBox.h"
 #include "EditorStyleSet.h"
+#include "PackageName.h"
+#include "AssetRegistryModule.h"
+#include "StringTable.h"
 #include "Internationalization/TextPackageNamespaceUtil.h"
+#include "Internationalization/StringTableCore.h"
+#include "Internationalization/StringTableRegistry.h"
 #include "Serialization/TextReferenceCollector.h"
 
 #define LOCTEXT_NAMESPACE "STextPropertyEditableTextBox"
@@ -74,6 +81,341 @@ void IEditableTextProperty::StaticStableTextId(UPackage* InPackage, const ETextP
 
 #endif // USE_STABLE_LOCALIZATION_KEYS
 
+void STextPropertyEditableStringTableReference::Construct(const FArguments& InArgs, const TSharedRef<IEditableTextProperty>& InEditableTextProperty)
+{
+	EditableTextProperty = InEditableTextProperty;
+
+	TSharedRef<SHorizontalBox> HorizontalBox = SNew(SHorizontalBox);
+
+	HorizontalBox->AddSlot()
+		[
+			SAssignNew(StringTableCombo, SComboBox<TSharedPtr<FAvailableStringTable>>)
+			.ComboBoxStyle(InArgs._ComboStyle)
+			.OptionsSource(&StringTableComboOptions)
+			.OnGenerateWidget(this, &STextPropertyEditableStringTableReference::MakeStringTableComboWidget)
+			.OnSelectionChanged(this, &STextPropertyEditableStringTableReference::OnStringTableComboChanged)
+			.OnComboBoxOpening(this, &STextPropertyEditableStringTableReference::OnStringTableComboOpening)
+			.Content()
+			[
+				SNew(STextBlock)
+				.Text(this, &STextPropertyEditableStringTableReference::GetStringTableComboContent)
+				.ToolTipText(this, &STextPropertyEditableStringTableReference::GetStringTableComboToolTip)
+			]
+		];
+
+	HorizontalBox->AddSlot()
+		[
+			SAssignNew(KeyCombo, SComboBox<TSharedPtr<FString>>)
+			.ComboBoxStyle(InArgs._ComboStyle)
+			.OptionsSource(&KeyComboOptions)
+			.OnGenerateWidget(this, &STextPropertyEditableStringTableReference::MakeKeyComboWidget)
+			.OnSelectionChanged(this, &STextPropertyEditableStringTableReference::OnKeyComboChanged)
+			.OnComboBoxOpening(this, &STextPropertyEditableStringTableReference::OnKeyComboOpening)
+			.Content()
+			[
+				SNew(STextBlock)
+				.Text(this, &STextPropertyEditableStringTableReference::GetKeyComboContent)
+				.ToolTipText(this, &STextPropertyEditableStringTableReference::GetKeyComboToolTip)
+			]
+		];
+
+	if (InArgs._AllowUnlink)
+	{
+		HorizontalBox->AddSlot()
+			.AutoWidth()
+			[
+				SNew(SButton)
+				.ButtonStyle(InArgs._ButtonStyle)
+				.Text(LOCTEXT("UnlinkStringTable", "Unlink"))
+				.IsEnabled(this, &STextPropertyEditableStringTableReference::IsUnlinkEnabled)
+				.OnClicked(this, &STextPropertyEditableStringTableReference::OnUnlinkClicked)
+			];
+	}
+
+	ChildSlot
+	[
+		HorizontalBox
+	];
+}
+
+void STextPropertyEditableStringTableReference::GetTableIdAndKey(FName& OutTableId, FString& OutKey) const
+{
+	const int32 NumTexts = EditableTextProperty->GetNumTexts();
+	if (NumTexts > 0)
+	{
+		const FText PropertyValue = EditableTextProperty->GetText(0);
+		FStringTableRegistry::Get().FindTableIdAndKey(PropertyValue, OutTableId, OutKey);
+
+		// Verify that all texts are using the same string table and key
+		for (int32 TextIndex = 1; TextIndex < NumTexts; ++TextIndex)
+		{
+			FName TmpTableId;
+			FString TmpKey;
+			if (FStringTableRegistry::Get().FindTableIdAndKey(PropertyValue, TmpTableId, TmpKey) && OutTableId == TmpTableId)
+			{
+				if (!OutKey.Equals(TmpKey, ESearchCase::CaseSensitive))
+				{
+					// Not using the same key - clear the key but keep the table and keep enumerating to verify the table on the remaining texts
+					OutKey.Reset();
+				}
+			}
+			else
+			{
+				// Not using a string table, or using a different string table - clear both table ID and key
+				OutTableId = FName();
+				OutKey.Reset();
+				break;
+			}
+		}
+	}
+}
+
+void STextPropertyEditableStringTableReference::SetTableIdAndKey(const FName InTableId, const FString& InKey)
+{
+	const FText TextToSet = FText::FromStringTable(InTableId, InKey);
+	if (TextToSet.IsFromStringTable())
+	{
+		const int32 NumTexts = EditableTextProperty->GetNumTexts();
+		for (int32 TextIndex = 0; TextIndex < NumTexts; ++TextIndex)
+		{
+			EditableTextProperty->SetText(TextIndex, TextToSet);
+		}
+	}
+}
+
+TSharedRef<SWidget> STextPropertyEditableStringTableReference::MakeStringTableComboWidget(TSharedPtr<FAvailableStringTable> InItem)
+{
+	return SNew(STextBlock)
+		.Text(InItem->DisplayName)
+		.ToolTipText(FText::FromName(InItem->TableId));
+}
+
+void STextPropertyEditableStringTableReference::OnStringTableComboChanged(TSharedPtr<FAvailableStringTable> NewSelection, ESelectInfo::Type SelectInfo)
+{
+	// If it's set from code, we did that on purpose
+	if (SelectInfo != ESelectInfo::Direct && NewSelection.IsValid())
+	{
+		// Make sure any selected string table asset is loaded
+		FName TableId = NewSelection->TableId;
+		IStringTableEngineBridge::RedirectAndLoadStringTableAsset(TableId, EStringTableLoadingPolicy::FindOrFullyLoad);
+
+		FStringTableConstPtr StringTable = FStringTableRegistry::Get().FindStringTable(TableId);
+		if (StringTable.IsValid())
+		{
+			// Just use the first key when changing the string table
+			StringTable->EnumerateSourceStrings([&](const FString& InKey, const FString& InSourceString) -> bool
+			{
+				SetTableIdAndKey(TableId, InKey);
+				return false; // stop enumeration
+			});
+		}
+	}
+}
+
+void STextPropertyEditableStringTableReference::OnStringTableComboOpening()
+{
+	FName CurrentTableId;
+	{
+		FString TmpKey;
+		GetTableIdAndKey(CurrentTableId, TmpKey);
+	}
+
+	TSharedPtr<FAvailableStringTable> SelectedStringTableComboEntry;
+	StringTableComboOptions.Reset();
+
+	// Process assets first (as they may currently be unloaded)
+	{
+		FAssetRegistryModule& AssetRegistryModule = FModuleManager::Get().LoadModuleChecked<FAssetRegistryModule>(AssetRegistryConstants::ModuleName);
+		
+		TArray<FAssetData> StringTableAssets;
+		AssetRegistryModule.Get().GetAssetsByClass(UStringTable::StaticClass()->GetFName(), StringTableAssets);
+		
+		for (const FAssetData& StringTableAsset : StringTableAssets)
+		{
+			TSharedRef<FAvailableStringTable> AvailableStringTableEntry = MakeShared<FAvailableStringTable>();
+			AvailableStringTableEntry->TableId = StringTableAsset.ObjectPath;
+			AvailableStringTableEntry->DisplayName = FText::FromName(StringTableAsset.AssetName);
+			if (StringTableAsset.ObjectPath == CurrentTableId)
+			{
+				SelectedStringTableComboEntry = AvailableStringTableEntry;
+			}
+			StringTableComboOptions.Add(AvailableStringTableEntry);
+		}
+	}
+
+	// Process the remaining non-asset string tables now
+	FStringTableRegistry::Get().EnumerateStringTables([&](const FName& InTableId, const FStringTableConstRef& InStringTable) -> bool
+	{
+		const bool bAlreadyAdded = StringTableComboOptions.ContainsByPredicate([InTableId](const TSharedPtr<FAvailableStringTable>& InAvailableStringTable)
+		{
+			return InAvailableStringTable->TableId == InTableId;
+		});
+
+		if (!bAlreadyAdded)
+		{
+			TSharedRef<FAvailableStringTable> AvailableStringTableEntry = MakeShared<FAvailableStringTable>();
+			AvailableStringTableEntry->TableId = InTableId;
+			AvailableStringTableEntry->DisplayName = FText::FromName(InTableId);
+			if (InTableId == CurrentTableId)
+			{
+				SelectedStringTableComboEntry = AvailableStringTableEntry;
+			}
+			StringTableComboOptions.Add(AvailableStringTableEntry);
+		}
+
+		return true; // continue enumeration
+	});
+
+	StringTableComboOptions.Sort([](const TSharedPtr<FAvailableStringTable>& InOne, const TSharedPtr<FAvailableStringTable>& InTwo)
+	{
+		return InOne->DisplayName.ToString() < InTwo->DisplayName.ToString();
+	});
+
+	if (SelectedStringTableComboEntry.IsValid())
+	{
+		StringTableCombo->SetSelectedItem(SelectedStringTableComboEntry);
+	}
+	else
+	{
+		StringTableCombo->ClearSelection();
+	}
+}
+
+FText STextPropertyEditableStringTableReference::GetStringTableComboContent() const
+{
+	FName CurrentTableId;
+	{
+		FString TmpKey;
+		GetTableIdAndKey(CurrentTableId, TmpKey);
+	}
+
+	return FText::FromString(FPackageName::GetLongPackageAssetName(CurrentTableId.ToString()));
+}
+
+FText STextPropertyEditableStringTableReference::GetStringTableComboToolTip() const
+{
+	FName CurrentTableId;
+	{
+		FString TmpKey;
+		GetTableIdAndKey(CurrentTableId, TmpKey);
+	}
+
+	return FText::FromName(CurrentTableId);
+}
+
+TSharedRef<SWidget> STextPropertyEditableStringTableReference::MakeKeyComboWidget(TSharedPtr<FString> InItem)
+{
+	return SNew(STextBlock)
+		.Text(FText::FromString(*InItem))
+		.ToolTipText(FText::FromString(*InItem));
+}
+
+void STextPropertyEditableStringTableReference::OnKeyComboChanged(TSharedPtr<FString> NewSelection, ESelectInfo::Type SelectInfo)
+{
+	// If it's set from code, we did that on purpose
+	if (SelectInfo != ESelectInfo::Direct && NewSelection.IsValid())
+	{
+		FName CurrentTableId;
+		{
+			FString TmpKey;
+			GetTableIdAndKey(CurrentTableId, TmpKey);
+		}
+
+		SetTableIdAndKey(CurrentTableId, *NewSelection);
+	}
+}
+
+void STextPropertyEditableStringTableReference::OnKeyComboOpening()
+{
+	FName CurrentTableId;
+	FString CurrentKey;
+	GetTableIdAndKey(CurrentTableId, CurrentKey);
+
+	TSharedPtr<FString> SelectedKeyComboEntry;
+	KeyComboOptions.Reset();
+
+	if (!CurrentTableId.IsNone())
+	{
+		FStringTableConstPtr StringTable = FStringTableRegistry::Get().FindStringTable(CurrentTableId);
+		if (StringTable.IsValid())
+		{
+			StringTable->EnumerateSourceStrings([&](const FString& InKey, const FString& InSourceString) -> bool
+			{
+				TSharedRef<FString> KeyComboEntry = MakeShared<FString>(InKey);
+				if (InKey.Equals(CurrentKey, ESearchCase::CaseSensitive))
+				{
+					SelectedKeyComboEntry = KeyComboEntry;
+				}
+				KeyComboOptions.Add(KeyComboEntry);
+				return true; // continue enumeration
+			});
+		}
+	}
+
+	KeyComboOptions.Sort([](const TSharedPtr<FString>& InOne, const TSharedPtr<FString>& InTwo)
+	{
+		return *InOne < *InTwo;
+	});
+
+	if (SelectedKeyComboEntry.IsValid())
+	{
+		KeyCombo->SetSelectedItem(SelectedKeyComboEntry);
+	}
+	else
+	{
+		KeyCombo->ClearSelection();
+	}
+}
+
+FText STextPropertyEditableStringTableReference::GetKeyComboContent() const
+{
+	FString CurrentKey;
+	{
+		FName TmpTableId;
+		GetTableIdAndKey(TmpTableId, CurrentKey);
+	}
+
+	return FText::FromString(MoveTemp(CurrentKey));
+}
+
+FText STextPropertyEditableStringTableReference::GetKeyComboToolTip() const
+{
+	return GetKeyComboContent();
+}
+
+bool STextPropertyEditableStringTableReference::IsUnlinkEnabled() const
+{
+	bool bEnabled = false;
+
+	const int32 NumTexts = EditableTextProperty->GetNumTexts();
+	for (int32 TextIndex = 0; TextIndex < NumTexts; ++TextIndex)
+	{
+		const FText CurrentText = EditableTextProperty->GetText(TextIndex);
+		if (CurrentText.IsFromStringTable())
+		{
+			bEnabled = true;
+			break;
+		}
+	}
+
+	return bEnabled;
+}
+
+FReply STextPropertyEditableStringTableReference::OnUnlinkClicked()
+{
+	const int32 NumTexts = EditableTextProperty->GetNumTexts();
+	for (int32 TextIndex = 0; TextIndex < NumTexts; ++TextIndex)
+	{
+		const FText CurrentText = EditableTextProperty->GetText(TextIndex);
+		if (CurrentText.IsFromStringTable())
+		{
+			EditableTextProperty->SetText(TextIndex, FText::GetEmpty());
+		}
+	}
+
+	return FReply::Handled();
+}
+
 void STextPropertyEditableTextBox::Construct(const FArguments& InArgs, const TSharedRef<IEditableTextProperty>& InEditableTextProperty)
 {
 	EditableTextProperty = InEditableTextProperty;
@@ -105,7 +447,7 @@ void STextPropertyEditableTextBox::Construct(const FArguments& InArgs, const TSh
 					.OnTextChanged(this, &STextPropertyEditableTextBox::OnTextChanged)
 					.OnTextCommitted(this, &STextPropertyEditableTextBox::OnTextCommitted)
 					.SelectAllTextOnCommit(false)
-					.IsReadOnly(this, &STextPropertyEditableTextBox::IsReadOnly)
+					.IsReadOnly(this, &STextPropertyEditableTextBox::IsSourceTextReadOnly)
 					.AutoWrapText(InArgs._AutoWrapText)
 					.WrapTextAt(InArgs._WrapTextAt)
 					.ModiferKeyForNewLine(EModifierKey::Shift)
@@ -138,7 +480,7 @@ void STextPropertyEditableTextBox::Construct(const FArguments& InArgs, const TSh
 					.OnTextChanged(this, &STextPropertyEditableTextBox::OnTextChanged)
 					.OnTextCommitted(this, &STextPropertyEditableTextBox::OnTextCommitted)
 					.SelectAllTextOnCommit(true)
-					.IsReadOnly(this, &STextPropertyEditableTextBox::IsReadOnly)
+					.IsReadOnly(this, &STextPropertyEditableTextBox::IsSourceTextReadOnly)
 					.IsPassword(bIsPassword)
 				]
 			]
@@ -166,15 +508,28 @@ void STextPropertyEditableTextBox::Construct(const FArguments& InArgs, const TSh
 					SNew(SGridPanel)
 					.FillColumn(1, 1.0f)
 
-					// Localizable?
+					// Inline Text
 					+SGridPanel::Slot(0, 0)
+					.ColumnSpan(2)
+					.Padding(2)
+					.HAlign(HAlign_Left)
+					.VAlign(VAlign_Center)
+					[
+						SNew(STextBlock)
+						.TextStyle(FEditorStyle::Get(), "LargeText")
+						.Text(LOCTEXT("TextInlineTextLabel", "Inline Text"))
+					]
+
+					// Localizable?
+					+SGridPanel::Slot(0, 1)
 					.Padding(2)
 					.HAlign(HAlign_Right)
+					.VAlign(VAlign_Center)
 					[
 						SNew(STextBlock)
 						.Text(LOCTEXT("TextLocalizableLabel", "Localizable:"))
 					]
-					+SGridPanel::Slot(1, 0)
+					+SGridPanel::Slot(1, 1)
 					.Padding(2)
 					[
 						SNew(SHorizontalBox)
@@ -193,6 +548,7 @@ void STextPropertyEditableTextBox::Construct(const FArguments& InArgs, const TSh
 								.ToolTipText(LOCTEXT("TextLocalizableToggleYesToolTip", "Assign this text a key and allow it to be gathered for localization."))
 								.Padding(FMargin(4, 2))
 								.HAlign(HAlign_Center)
+								.IsEnabled(this, &STextPropertyEditableTextBox::IsCultureInvariantFlagEnabled)
 								.IsChecked(this, &STextPropertyEditableTextBox::GetLocalizableCheckState, true/*bActiveState*/)
 								.OnCheckStateChanged(this, &STextPropertyEditableTextBox::HandleLocalizableCheckStateChanged, true/*bActiveState*/)
 								[
@@ -208,6 +564,7 @@ void STextPropertyEditableTextBox::Construct(const FArguments& InArgs, const TSh
 								.ToolTipText(LOCTEXT("TextLocalizableToggleNoToolTip", "Mark this text as 'culture invariant' to prevent it being gathered for localization."))
 								.Padding(FMargin(4, 2))
 								.HAlign(HAlign_Center)
+								.IsEnabled(this, &STextPropertyEditableTextBox::IsCultureInvariantFlagEnabled)
 								.IsChecked(this, &STextPropertyEditableTextBox::GetLocalizableCheckState, false/*bActiveState*/)
 								.OnCheckStateChanged(this, &STextPropertyEditableTextBox::HandleLocalizableCheckStateChanged, false/*bActiveState*/)
 								[
@@ -220,14 +577,15 @@ void STextPropertyEditableTextBox::Construct(const FArguments& InArgs, const TSh
 
 #if USE_STABLE_LOCALIZATION_KEYS
 					// Package
-					+SGridPanel::Slot(0, 1)
+					+SGridPanel::Slot(0, 2)
 					.Padding(2)
 					.HAlign(HAlign_Right)
+					.VAlign(VAlign_Center)
 					[
 						SNew(STextBlock)
 						.Text(LOCTEXT("TextPackageLabel", "Package:"))
 					]
-					+SGridPanel::Slot(1, 1)
+					+SGridPanel::Slot(1, 2)
 					.Padding(2)
 					[
 						SNew(SEditableTextBox)
@@ -237,14 +595,15 @@ void STextPropertyEditableTextBox::Construct(const FArguments& InArgs, const TSh
 #endif // USE_STABLE_LOCALIZATION_KEYS
 
 					// Namespace
-					+SGridPanel::Slot(0, 2)
+					+SGridPanel::Slot(0, 3)
 					.Padding(2)
 					.HAlign(HAlign_Right)
+					.VAlign(VAlign_Center)
 					[
 						SNew(STextBlock)
 						.Text(LOCTEXT("TextNamespaceLabel", "Namespace:"))
 					]
-					+SGridPanel::Slot(1, 2)
+					+SGridPanel::Slot(1, 3)
 					.Padding(2)
 					[
 						SAssignNew(NamespaceEditableTextBox, SEditableTextBox)
@@ -258,14 +617,15 @@ void STextPropertyEditableTextBox::Construct(const FArguments& InArgs, const TSh
 					]
 
 					// Key
-					+SGridPanel::Slot(0, 3)
+					+SGridPanel::Slot(0, 4)
 					.Padding(2)
 					.HAlign(HAlign_Right)
+					.VAlign(VAlign_Center)
 					[
 						SNew(STextBlock)
 						.Text(LOCTEXT("TextKeyLabel", "Key:"))
 					]
-					+SGridPanel::Slot(1, 3)
+					+SGridPanel::Slot(1, 4)
 					.Padding(2)
 					[
 						SAssignNew(KeyEditableTextBox, SEditableTextBox)
@@ -281,6 +641,35 @@ void STextPropertyEditableTextBox::Construct(const FArguments& InArgs, const TSh
 						.IsReadOnly(true)
 #endif	// USE_STABLE_LOCALIZATION_KEYS
 					]
+
+					// Referenced Text
+					+SGridPanel::Slot(0, 5)
+					.ColumnSpan(2)
+					.Padding(2)
+					.HAlign(HAlign_Left)
+					.VAlign(VAlign_Center)
+					[
+						SNew(STextBlock)
+						.TextStyle(FEditorStyle::Get(), "LargeText")
+						.Text(LOCTEXT("TextReferencedTextLabel", "Referenced Text"))
+					]
+
+					// String Table
+					+SGridPanel::Slot(0, 6)
+					.Padding(2)
+					.HAlign(HAlign_Right)
+					.VAlign(VAlign_Center)
+					[
+						SNew(STextBlock)
+						.Text(LOCTEXT("TextStringTableLabel", "String Table:"))
+					]
+					+SGridPanel::Slot(1, 6)
+					.Padding(2)
+					[
+						SNew(STextPropertyEditableStringTableReference, InEditableTextProperty)
+						.AllowUnlink(true)
+						.IsEnabled(this, &STextPropertyEditableTextBox::CanEdit)
+					]
 				]
 			]
 		];
@@ -339,9 +728,30 @@ bool STextPropertyEditableTextBox::CanEdit() const
 	return !EditableTextProperty->IsReadOnly();
 }
 
-bool STextPropertyEditableTextBox::IsReadOnly() const
+bool STextPropertyEditableTextBox::IsCultureInvariantFlagEnabled() const
 {
-	return EditableTextProperty->IsReadOnly();
+	return !IsSourceTextReadOnly();
+}
+
+bool STextPropertyEditableTextBox::IsSourceTextReadOnly() const
+{
+	if (EditableTextProperty->IsReadOnly())
+	{
+		return true;
+	}
+
+	// We can't edit the source string of string table references
+	const int32 NumTexts = EditableTextProperty->GetNumTexts();
+	if (NumTexts == 1)
+	{
+		const FText TextValue = EditableTextProperty->GetText(0);
+		if (TextValue.IsFromStringTable())
+		{
+			return true;
+		}
+	}
+
+	return false;
 }
 
 bool STextPropertyEditableTextBox::IsIdentityReadOnly() const
@@ -351,12 +761,12 @@ bool STextPropertyEditableTextBox::IsIdentityReadOnly() const
 		return true;
 	}
 
-	// We can't edit the identity of culture invariant texts
+	// We can't edit the identity of culture invariant texts, or string table references
 	const int32 NumTexts = EditableTextProperty->GetNumTexts();
 	if (NumTexts == 1)
 	{
 		const FText TextValue = EditableTextProperty->GetText(0);
-		if (TextValue.IsCultureInvariant())
+		if (TextValue.IsCultureInvariant() || TextValue.IsFromStringTable())
 		{
 			return true;
 		}
@@ -373,26 +783,40 @@ FText STextPropertyEditableTextBox::GetToolTipText() const
 	{
 		const FText TextValue = EditableTextProperty->GetText(0);
 
-		bool bIsLocalized = false;
-		FString Namespace;
-		FString Key;
-		const FString* SourceString = FTextInspector::GetSourceString(TextValue);
-
-		if (SourceString && TextValue.ShouldGatherForLocalization())
+		if (TextValue.IsFromStringTable())
 		{
-			bIsLocalized = FTextLocalizationManager::Get().FindNamespaceAndKeyFromDisplayString(FTextInspector::GetSharedDisplayString(TextValue), Namespace, Key);
-		}
-
-		if (bIsLocalized)
-		{
-			const FString PackageNamespace = TextNamespaceUtil::ExtractPackageNamespace(Namespace);
-			const FString TextNamespace = TextNamespaceUtil::StripPackageNamespace(Namespace);
+			FName TableId;
+			FString Key;
+			FStringTableRegistry::Get().FindTableIdAndKey(TextValue, TableId, Key);
 
 			LocalizedTextToolTip = FText::Format(
-				LOCTEXT("LocalizedTextToolTipFmt", "--- Localized Text ---\nPackage: {0}\nNamespace: {1}\nKey: {2}\nSource: {3}"), 
-				FText::FromString(PackageNamespace), FText::FromString(TextNamespace), FText::FromString(Key), FText::FromString(*SourceString)
+				LOCTEXT("StringTableTextToolTipFmt", "--- String Table Reference ---\nTable ID: {0}\nKey: {1}"), 
+				FText::FromName(TableId), FText::FromString(Key)
 				);
 		}
+		else
+		{
+			bool bIsLocalized = false;
+			FString Namespace;
+			FString Key;
+			const FString* SourceString = FTextInspector::GetSourceString(TextValue);
+
+			if (SourceString && TextValue.ShouldGatherForLocalization())
+			{
+				bIsLocalized = FTextLocalizationManager::Get().FindNamespaceAndKeyFromDisplayString(FTextInspector::GetSharedDisplayString(TextValue), Namespace, Key);
+			}
+
+			if (bIsLocalized)
+			{
+				const FString PackageNamespace = TextNamespaceUtil::ExtractPackageNamespace(Namespace);
+				const FString TextNamespace = TextNamespaceUtil::StripPackageNamespace(Namespace);
+
+				LocalizedTextToolTip = FText::Format(
+					LOCTEXT("LocalizedTextToolTipFmt", "--- Localized Text ---\nPackage: {0}\nNamespace: {1}\nKey: {2}\nSource: {3}"), 
+					FText::FromString(PackageNamespace), FText::FromString(TextNamespace), FText::FromString(Key), FText::FromString(*SourceString)
+					);
+			}
+		}
 	}
 	
 	const FText BaseToolTipText = EditableTextProperty->GetToolTipText();
diff --git a/Engine/Source/Editor/EditorWidgets/Public/STextPropertyEditableTextBox.h b/Engine/Source/Editor/EditorWidgets/Public/STextPropertyEditableTextBox.h
--- a/Engine/Source/Editor/EditorWidgets/Public/STextPropertyEditableTextBox.h
+++ b/Engine/Source/Editor/EditorWidgets/Public/STextPropertyEditableTextBox.h
@@ -18,6 +18,9 @@
 class SEditableTextBox;
 class SMultiLineEditableTextBox;
 
+template <typename OptionType>
+class SComboBox;
+
 /** Interface to allow STextPropertyEditableTextBox to be used to edit both properties and Blueprint pins */
 class IEditableTextProperty
 {
@@ -76,6 +79,59 @@ protected:
 #endif // USE_STABLE_LOCALIZATION_KEYS
 };
 
+/** A widget that can be used for editing the string table referenced for FText instances */
+class EDITORWIDGETS_API STextPropertyEditableStringTableReference : public SCompoundWidget
+{
+	SLATE_BEGIN_ARGS(STextPropertyEditableStringTableReference)
+		: _ComboStyle(&FCoreStyle::Get().GetWidgetStyle<FComboBoxStyle>("ComboBox"))
+		, _ButtonStyle(&FCoreStyle::Get().GetWidgetStyle<FButtonStyle>("Button"))
+		, _AllowUnlink(false)
+		{}
+		/** The styling of the combobox */
+		SLATE_STYLE_ARGUMENT(FComboBoxStyle, ComboStyle)
+		/** The styling of the button */
+		SLATE_STYLE_ARGUMENT(FButtonStyle, ButtonStyle)
+		/** Should we show an "unlink" button? */
+		SLATE_ARGUMENT(bool, AllowUnlink)
+	SLATE_END_ARGS()
+
+public:
+	void Construct(const FArguments& Arguments, const TSharedRef<IEditableTextProperty>& InEditableTextProperty);
+
+private:
+	struct FAvailableStringTable
+	{
+		FName TableId;
+		FText DisplayName;
+	};
+
+	void GetTableIdAndKey(FName& OutTableId, FString& OutKey) const;
+	void SetTableIdAndKey(const FName InTableId, const FString& InKey);
+
+	TSharedRef<SWidget> MakeStringTableComboWidget(TSharedPtr<FAvailableStringTable> InItem);
+	void OnStringTableComboChanged(TSharedPtr<FAvailableStringTable> NewSelection, ESelectInfo::Type SelectInfo);
+	void OnStringTableComboOpening();
+	FText GetStringTableComboContent() const;
+	FText GetStringTableComboToolTip() const;
+
+	TSharedRef<SWidget> MakeKeyComboWidget(TSharedPtr<FString> InItem);
+	void OnKeyComboChanged(TSharedPtr<FString> NewSelection, ESelectInfo::Type SelectInfo);
+	void OnKeyComboOpening();
+	FText GetKeyComboContent() const;
+	FText GetKeyComboToolTip() const;
+
+	bool IsUnlinkEnabled() const;
+	FReply OnUnlinkClicked();
+
+	TSharedPtr<IEditableTextProperty> EditableTextProperty;
+
+	TSharedPtr<SComboBox<TSharedPtr<FAvailableStringTable>>> StringTableCombo;
+	TArray<TSharedPtr<FAvailableStringTable>> StringTableComboOptions;
+
+	TSharedPtr<SComboBox<TSharedPtr<FString>>> KeyCombo;
+	TArray<TSharedPtr<FString>> KeyComboOptions;
+};
+
 /** A widget that can be used for editing FText instances */
 class EDITORWIDGETS_API STextPropertyEditableTextBox : public SCompoundWidget
 {
@@ -113,7 +169,8 @@ public:
 private:
 	void GetDesiredWidth(float& OutMinDesiredWidth, float& OutMaxDesiredWidth);
 	bool CanEdit() const;
-	bool IsReadOnly() const;
+	bool IsCultureInvariantFlagEnabled() const;
+	bool IsSourceTextReadOnly() const;
 	bool IsIdentityReadOnly() const;
 	FText GetToolTipText() const;
 
diff --git a/Engine/Source/Editor/FoliageEdit/Private/FoliageEdMode.cpp b/Engine/Source/Editor/FoliageEdit/Private/FoliageEdMode.cpp
--- a/Engine/Source/Editor/FoliageEdit/Private/FoliageEdMode.cpp
+++ b/Engine/Source/Editor/FoliageEdit/Private/FoliageEdMode.cpp
@@ -2870,20 +2870,37 @@ UFoliageType* FEdModeFoliage::AddFoliageAsset(UObject* InAsset)
 {
 	UFoliageType* FoliageType = nullptr;
 
-	const FScopedTransaction Transaction(NSLOCTEXT("UnrealEd", "FoliageMode_AddTypeTransaction", "Add Foliage Type"));
-
 	UStaticMesh* StaticMesh = Cast<UStaticMesh>(InAsset);
 	if (StaticMesh)
 	{
-		AInstancedFoliageActor* IFA = AInstancedFoliageActor::GetInstancedFoliageActorForCurrentLevel(GetWorld(), true);
-		FoliageType = IFA->GetLocalFoliageTypeForMesh(StaticMesh);
-		if (!FoliageType)
+		UWorld* World = GetWorld();
+
 		{
-			IFA->AddMesh(StaticMesh, &FoliageType);
+			const FScopedTransaction Transaction(NSLOCTEXT("UnrealEd", "FoliageMode_AddTypeTransaction", "Add Foliage Type"));
+
+			AInstancedFoliageActor* IFA = AInstancedFoliageActor::GetInstancedFoliageActorForCurrentLevel(World, true);
+			FoliageType = IFA->GetLocalFoliageTypeForMesh(StaticMesh);
+			if (!FoliageType)
+			{
+				IFA->AddMesh(StaticMesh, &FoliageType);
+			}
+		}
+
+		// If there is multiple levels for this world, save the foliage directly as an asset, so user will be able to paint over all levels by default
+		if (World->StreamingLevels.Num() > 0)
+		{
+			UFoliageType* TypeSaved = SaveFoliageTypeObject(FoliageType);
+
+			if (TypeSaved != nullptr)
+			{
+				FoliageType = TypeSaved;
+			}
 		}
 	}
 	else
 	{
+		const FScopedTransaction Transaction(NSLOCTEXT("UnrealEd", "FoliageMode_AddTypeTransaction", "Add Foliage Type"));
+
 		FoliageType = Cast<UFoliageType>(InAsset);
 		if (FoliageType)
 		{
@@ -3040,8 +3057,6 @@ void FEdModeFoliage::ReplaceSettingsObject(UFoliageType* OldSettings, UFoliageTy
 
 UFoliageType* FEdModeFoliage::SaveFoliageTypeObject(UFoliageType* InFoliageType)
 {
-	FScopedTransaction Transaction(NSLOCTEXT("UnrealEd", "FoliageMode_SaveFoliageTypeObject", "Foliage Editing: Save Foliage Type Object"));
-
 	UFoliageType* TypeToSave = nullptr;
 
 	if (!InFoliageType->IsAsset())
@@ -3063,8 +3078,13 @@ UFoliageType* FEdModeFoliage::SaveFoliageTypeObject(UFoliageType* InFoliageType)
 		{
 			PackageName = SaveFoliageTypeDialog->GetFullAssetPath().ToString();
 			UPackage* Package = CreatePackage(nullptr, *PackageName);
+
+			// We should not save a copy of this duplicate into the transaction buffer as it's an asset
+			InFoliageType->ClearFlags(RF_Transactional);
 			TypeToSave = Cast<UFoliageType>(StaticDuplicateObject(InFoliageType, Package, *FPackageName::GetLongPackageAssetName(PackageName)));
-			TypeToSave->SetFlags(RF_Standalone | RF_Public);
+			InFoliageType->SetFlags(RF_Transactional);
+
+			TypeToSave->SetFlags(RF_Standalone | RF_Public | RF_Transactional);
 			TypeToSave->Modify();
 
 			// Notify the asset registry
diff --git a/Engine/Source/Editor/FoliageEdit/Private/SFoliagePalette.cpp b/Engine/Source/Editor/FoliageEdit/Private/SFoliagePalette.cpp
--- a/Engine/Source/Editor/FoliageEdit/Private/SFoliagePalette.cpp
+++ b/Engine/Source/Editor/FoliageEdit/Private/SFoliagePalette.cpp
@@ -162,7 +162,8 @@ void SFoliagePalette::Construct(const FArguments& InArgs)
 {
 	bItemsNeedRebuild = false;
 	bIsUneditableFoliageTypeSelected = false;
-	bIsActiveTimerRegistered = false;
+	bIsRebuildTimerRegistered = false;
+	bIsRefreshTimerRegistered = false;
 
 	FoliageEditMode = InArgs._FoliageEdMode;
 
@@ -172,7 +173,9 @@ void SFoliagePalette::Construct(const FArguments& InArgs)
 	UICommandList = MakeShareable(new FUICommandList);
 	BindCommands();
 
-	ThumbnailPool = MakeShareable(new FAssetThumbnailPool(25, TAttribute<bool>::Create(TAttribute<bool>::FGetter::CreateSP(this, &SFoliagePalette::IsHovered))));
+	// Size of the thumbnail pool should be large enough to show a reasonable amount of foliage assets on screen at once,
+	// otherwise some thumbnail images will appear duplicated.
+	ThumbnailPool = MakeShareable(new FAssetThumbnailPool(64, TAttribute<bool>::Create(TAttribute<bool>::FGetter::CreateSP(this, &SFoliagePalette::IsHovered))));
 
 	TypeFilter = MakeShareable(new FoliageTypeTextFilter(
 		FoliageTypeTextFilter::FItemToStringArray::CreateSP(this, &SFoliagePalette::GetPaletteItemFilterString)));
@@ -433,13 +436,23 @@ void SFoliagePalette::UpdatePalette(bool bRebuildItems)
 {
 	bItemsNeedRebuild |= bRebuildItems;
 
-	if (!bIsActiveTimerRegistered)
+	if (!bIsRebuildTimerRegistered)
 	{
-		bIsActiveTimerRegistered = true;
+		bIsRebuildTimerRegistered = true;
 		RegisterActiveTimer(0.f, FWidgetActiveTimerDelegate::CreateSP(this, &SFoliagePalette::UpdatePaletteItems));
 	}
 }
 
+void SFoliagePalette::RefreshPalette()
+{
+	// Do not register the refresh timer if we're pending a rebuild; rebuild should cause the palette to refresh
+	if (!bIsRefreshTimerRegistered && !bIsRebuildTimerRegistered)
+	{
+		bIsRefreshTimerRegistered = true;
+		RegisterActiveTimer(0.f, FWidgetActiveTimerDelegate::CreateSP(this, &SFoliagePalette::RefreshPaletteItems));
+	}
+}
+
 void SFoliagePalette::UpdateThumbnailForType(UFoliageType* FoliageType)
 {
 	// Recreate the palette item for the given foliage type
@@ -1169,6 +1182,9 @@ bool SFoliagePalette::CanSelectInstances() const
 TSharedRef<ITableRow> SFoliagePalette::GenerateTile(FFoliagePaletteItemModelPtr Item, const TSharedRef<STableViewBase>& OwnerTable)
 {
 	return SNew(SFoliagePaletteItemTile, OwnerTable, Item);
+
+	// Refresh the palette to ensure that thumbnails are correct
+	RefreshPalette();
 }
 
 float SFoliagePalette::GetScaledThumbnailSize() const
@@ -1397,7 +1413,19 @@ EActiveTimerReturnType SFoliagePalette::UpdatePaletteItems(double InCurrentTime,
 	// Refresh the appropriate view
 	RefreshActivePaletteViewWidget();
 
-	bIsActiveTimerRegistered = false;
+	bIsRebuildTimerRegistered = false;
+	return EActiveTimerReturnType::Stop;
+}
+
+EActiveTimerReturnType SFoliagePalette::RefreshPaletteItems(double InCurrentTime, float InDeltaTime)
+{
+	// Do not refresh the palette if we're waiting on a rebuild
+	if (!bItemsNeedRebuild)
+	{
+		RefreshActivePaletteViewWidget();
+	}
+
+	bIsRefreshTimerRegistered = false;
 	return EActiveTimerReturnType::Stop;
 }
 
diff --git a/Engine/Source/Editor/FoliageEdit/Private/SFoliagePalette.h b/Engine/Source/Editor/FoliageEdit/Private/SFoliagePalette.h
--- a/Engine/Source/Editor/FoliageEdit/Private/SFoliagePalette.h
+++ b/Engine/Source/Editor/FoliageEdit/Private/SFoliagePalette.h
@@ -51,6 +51,9 @@ public:
 	/** Updates the foliage palette, optionally doing a full rebuild of the items in the palette as well */
 	void UpdatePalette(bool bRebuildItems = false);
 
+	/** Refreshes the foliage palette */
+	void RefreshPalette();
+
 	/** Updates the thumbnail for the given foliage type in the palette */
 	void UpdateThumbnailForType(UFoliageType* FoliageType);
 
@@ -242,6 +245,9 @@ private:
 	/** Active timer handler to update the items in the palette */
 	EActiveTimerReturnType UpdatePaletteItems(double InCurrentTime, float InDeltaTime);
 
+	/** Active timer handler to refresh the palette */
+	EActiveTimerReturnType RefreshPaletteItems(double InCurrentTime, float InDeltaTime);
+
 private:
 	typedef TTextFilter<FFoliagePaletteItemModelPtr> FoliageTypeTextFilter;
 	TSharedPtr<FoliageTypeTextFilter> TypeFilter;
@@ -283,5 +289,7 @@ private:
 
 	bool bItemsNeedRebuild : 1;
 	bool bIsUneditableFoliageTypeSelected : 1;
-	bool bIsActiveTimerRegistered : 1;
+	
+	bool bIsRebuildTimerRegistered : 1;
+	bool bIsRefreshTimerRegistered : 1;
 };
diff --git a/Engine/Source/Editor/GraphEditor/Private/KismetNodes/SGraphNodeK2Sequence.h b/Engine/Source/Editor/GraphEditor/Private/KismetNodes/SGraphNodeK2Sequence.h
--- a/Engine/Source/Editor/GraphEditor/Private/KismetNodes/SGraphNodeK2Sequence.h
+++ b/Engine/Source/Editor/GraphEditor/Private/KismetNodes/SGraphNodeK2Sequence.h
@@ -1,8 +1,6 @@
 // Copyright 1998-2017 Epic Games, Inc. All Rights Reserved.
 
-
-#ifndef __SGraphNodeK2Sequence_h__
-#define __SGraphNodeK2Sequence_h__
+#pragma once
 
 #include "CoreMinimal.h"
 #include "Widgets/DeclarativeSyntaxSupport.h"
@@ -26,5 +24,3 @@ protected:
 	virtual FReply OnAddPin() override;
 	// End of SGraphNode interface
 };
-
-#endif // __SGraphNodeK2Sequence_h__
diff --git a/Engine/Source/Editor/GraphEditor/Private/KismetNodes/SGraphNodeK2Terminator.h b/Engine/Source/Editor/GraphEditor/Private/KismetNodes/SGraphNodeK2Terminator.h
--- a/Engine/Source/Editor/GraphEditor/Private/KismetNodes/SGraphNodeK2Terminator.h
+++ b/Engine/Source/Editor/GraphEditor/Private/KismetNodes/SGraphNodeK2Terminator.h
@@ -1,8 +1,6 @@
 // Copyright 1998-2017 Epic Games, Inc. All Rights Reserved.
 
-
-#ifndef __SGraphNodeK2Terminator_h__
-#define __SGraphNodeK2Terminator_h__
+#pragma once
 
 #include "CoreMinimal.h"
 #include "Widgets/DeclarativeSyntaxSupport.h"
@@ -23,5 +21,3 @@ public:
 protected:
 	virtual void UpdateGraphNode() override;
 };
-
-#endif // __SGraphNodeK2Terminator_h__
diff --git a/Engine/Source/Editor/GraphEditor/Private/KismetNodes/SGraphNodeK2Var.h b/Engine/Source/Editor/GraphEditor/Private/KismetNodes/SGraphNodeK2Var.h
--- a/Engine/Source/Editor/GraphEditor/Private/KismetNodes/SGraphNodeK2Var.h
+++ b/Engine/Source/Editor/GraphEditor/Private/KismetNodes/SGraphNodeK2Var.h
@@ -1,8 +1,6 @@
 // Copyright 1998-2017 Epic Games, Inc. All Rights Reserved.
 
-
-#ifndef __SGraphNodeK2Var_h__
-#define __SGraphNodeK2Var_h__
+#pragma once
 
 #include "CoreMinimal.h"
 #include "Styling/SlateColor.h"
@@ -29,5 +27,3 @@ public:
 protected:
 	FSlateColor GetVariableColor() const;
 };
-
-#endif // __SGraphNodeK2Var_h__
diff --git a/Engine/Source/Editor/GraphEditor/Private/SGraphEditorActionMenu.h b/Engine/Source/Editor/GraphEditor/Private/SGraphEditorActionMenu.h
--- a/Engine/Source/Editor/GraphEditor/Private/SGraphEditorActionMenu.h
+++ b/Engine/Source/Editor/GraphEditor/Private/SGraphEditorActionMenu.h
@@ -1,8 +1,6 @@
 // Copyright 1998-2017 Epic Games, Inc. All Rights Reserved.
 
-
-#ifndef __SGraphEditorActionMenu_h__
-#define __SGraphEditorActionMenu_h__
+#pragma once
 
 #include "CoreMinimal.h"
 #include "Widgets/DeclarativeSyntaxSupport.h"
@@ -53,5 +51,3 @@ protected:
 	/** Callback used to populate all actions list in SGraphActionMenu */
 	void CollectAllActions(FGraphActionListBuilderBase& OutAllActions);
 };
-
-#endif // __SGraphEditorActionMenu_h__
diff --git a/Engine/Source/Editor/GraphEditor/Private/SGraphNodeDefault.h b/Engine/Source/Editor/GraphEditor/Private/SGraphNodeDefault.h
--- a/Engine/Source/Editor/GraphEditor/Private/SGraphNodeDefault.h
+++ b/Engine/Source/Editor/GraphEditor/Private/SGraphNodeDefault.h
@@ -1,8 +1,6 @@
 // Copyright 1998-2017 Epic Games, Inc. All Rights Reserved.
 
-
-#ifndef __SGraphNodeDefault_h__
-#define __SGraphNodeDefault_h__
+#pragma once
 
 #include "CoreMinimal.h"
 #include "Widgets/DeclarativeSyntaxSupport.h"
@@ -21,5 +19,3 @@ public:
 
 	void Construct( const FArguments& InArgs );
 };
-
-#endif // __SGraphNodeDefault_h__
diff --git a/Engine/Source/Editor/GraphEditor/Public/GraphEditorSettings.h b/Engine/Source/Editor/GraphEditor/Public/GraphEditorSettings.h
--- a/Engine/Source/Editor/GraphEditor/Public/GraphEditorSettings.h
+++ b/Engine/Source/Editor/GraphEditor/Public/GraphEditorSettings.h
@@ -34,23 +34,23 @@ public:
 	TEnumAsByte<EBlueprintPinStyleType> DataPinStyle;
 
 	/** The amount of padding above a pin (defaults to 4) */
-	UPROPERTY(/*config, */ EditAnywhere, Category=GeneralStyle, AdvancedDisplay)
+	UPROPERTY(config, EditAnywhere, Category=GeneralStyle, AdvancedDisplay)
 	float PaddingAbovePin;
 
 	/** The amount of padding below a pin (defaults to 4) */
-	UPROPERTY(/*config, */ EditAnywhere, Category=GeneralStyle, AdvancedDisplay)
+	UPROPERTY(config, EditAnywhere, Category=GeneralStyle, AdvancedDisplay)
 	float PaddingBelowPin;
 
 	/** The amount of horizontal padding towards the center of a node on an input pin (defaults to 10) */
-	UPROPERTY(/*config, */ EditAnywhere, Category=GeneralStyle, AdvancedDisplay)
+	UPROPERTY(config, EditAnywhere, Category=GeneralStyle, AdvancedDisplay)
 	float PaddingRightOfInput;
 
 	/** The amount of horizontal padding towards the center of a node on an output pin (defaults to 10) */
-	UPROPERTY(/*config, */ EditAnywhere, Category=GeneralStyle, AdvancedDisplay)
+	UPROPERTY(config, EditAnywhere, Category=GeneralStyle, AdvancedDisplay)
 	float PaddingLeftOfOutput;
 
 	/** The amount of padding towards the node edge (defaults to 10, can be negative to make pins overlap or stick out of a node */
-	UPROPERTY(/*config, */ EditAnywhere, Category=GeneralStyle, AdvancedDisplay)
+	UPROPERTY(config, EditAnywhere, Category=GeneralStyle, AdvancedDisplay)
 	float PaddingTowardsNodeEdge;
 
 	/** Returns the padding to use for an input pin on a graph node */
diff --git a/Engine/Source/Editor/Kismet/Private/BPVariableDragDropAction.h b/Engine/Source/Editor/Kismet/Private/BPVariableDragDropAction.h
--- a/Engine/Source/Editor/Kismet/Private/BPVariableDragDropAction.h
+++ b/Engine/Source/Editor/Kismet/Private/BPVariableDragDropAction.h
@@ -61,9 +61,7 @@ public:
 	{
 		if (VariableSource.IsValid() && VariableName != NAME_None)
 		{
-			UProperty* VariableProperty = FindField<UProperty>(VariableSource.Get(), VariableName);
-			check(VariableProperty != NULL);
-			return VariableProperty;
+			return FindFieldChecked<UProperty>(VariableSource.Get(), VariableName);
 		}
 		return nullptr;
 	}
diff --git a/Engine/Source/Editor/Kismet/Private/UserDefinedStructureEditor.cpp b/Engine/Source/Editor/Kismet/Private/UserDefinedStructureEditor.cpp
--- a/Engine/Source/Editor/Kismet/Private/UserDefinedStructureEditor.cpp
+++ b/Engine/Source/Editor/Kismet/Private/UserDefinedStructureEditor.cpp
@@ -382,7 +382,7 @@ TSharedRef<SDockTab> FUserDefinedStructureEditor::SpawnStructureTab(const FSpawn
 	TSharedRef<SSplitter> Splitter = SNew(SSplitter)
 		.Orientation(Orient_Vertical)
 		.PhysicalSplitterHandleSize(10.0f)
-		.ResizeMode(ESplitterResizeMode::Fixed);
+		.ResizeMode(ESplitterResizeMode::FixedPosition);
 
 	{
 		// Create a property view
diff --git a/Engine/Source/Editor/KismetCompiler/Private/KismetCompiler.cpp b/Engine/Source/Editor/KismetCompiler/Private/KismetCompiler.cpp
--- a/Engine/Source/Editor/KismetCompiler/Private/KismetCompiler.cpp
+++ b/Engine/Source/Editor/KismetCompiler/Private/KismetCompiler.cpp
@@ -3997,7 +3997,7 @@ void FKismetCompilerContext::Compile()
 		for (UBlueprint* CurrentBP : DependentBlueprints)
 		{
 			// Get the current dirty state of the package
-			UPackage* const Package = Cast<UPackage>(CurrentBP->GetOutermost());
+			UPackage* const Package = CurrentBP->GetOutermost();
 			const bool bStartedWithUnsavedChanges = Package != nullptr ? Package->IsDirty() : true;
 			const EBlueprintStatus OriginalStatus = CurrentBP->Status;
 
diff --git a/Engine/Source/Editor/KismetCompiler/Private/KismetCompilerBackend.h b/Engine/Source/Editor/KismetCompiler/Private/KismetCompilerBackend.h
--- a/Engine/Source/Editor/KismetCompiler/Private/KismetCompilerBackend.h
+++ b/Engine/Source/Editor/KismetCompiler/Private/KismetCompilerBackend.h
@@ -1,9 +1,5 @@
 // Copyright 1998-2017 Epic Games, Inc. All Rights Reserved.
 
-
-#ifndef __KismetCompilerBackend_h__
-#define __KismetCompilerBackend_h__
-
 #pragma once
 
 #include "CoreMinimal.h"
@@ -52,5 +48,3 @@ protected:
 };
 
 //////////////////////////////////////////////////////////////////////////
-
-#endif	// __KismetCompilerBackend_h__
diff --git a/Engine/Source/Editor/KismetCompiler/Private/KismetCompilerVMBackend.cpp b/Engine/Source/Editor/KismetCompiler/Private/KismetCompilerVMBackend.cpp
--- a/Engine/Source/Editor/KismetCompiler/Private/KismetCompilerVMBackend.cpp
+++ b/Engine/Source/Editor/KismetCompiler/Private/KismetCompilerVMBackend.cpp
@@ -14,7 +14,6 @@
 #include "UObject/Interface.h"
 #include "UObject/UnrealType.h"
 #include "UObject/TextProperty.h"
-#include "Internationalization/TextNamespaceUtil.h"
 #include "UObject/PropertyPortFlags.h"
 #include "EdGraph/EdGraphNode.h"
 #include "EdGraph/EdGraphPin.h"
@@ -36,7 +35,12 @@
 #include "Kismet2/StructureEditorUtils.h"
 #include "Kismet2/KismetDebugUtilities.h"
 
+#include "Internationalization/TextNamespaceUtil.h"
 #include "Internationalization/TextPackageNamespaceUtil.h"
+#include "Internationalization/StringTableRegistry.h"
+#include "Internationalization/StringTableCore.h"
+#include "Internationalization/StringTable.h"
+
 #include "Kismet2/BlueprintEditorUtils.h"
 
 #define LOCTEXT_NAMESPACE "KismetCompilerVMBackend"
@@ -628,6 +632,19 @@ public:
 				{
 					Writer << EBlueprintTextLiteralType::Empty;
 				}
+				else if (Term->TextLiteral.IsFromStringTable())
+				{
+					FName TableId;
+					FString Key;
+					FStringTableRegistry::Get().FindTableIdAndKey(Term->TextLiteral, TableId, Key);
+
+					UStringTable* StringTableAsset = FStringTableRegistry::Get().FindStringTableAsset(TableId);
+
+					Writer << EBlueprintTextLiteralType::StringTableEntry;
+					Writer << StringTableAsset; // Not used at runtime, but exists for asset dependency gathering
+					EmitStringLiteral(TableId.ToString());
+					EmitStringLiteral(Key);
+				}
 				else if (Term->TextLiteral.IsCultureInvariant())
 				{
 					Writer << EBlueprintTextLiteralType::InvariantText;
@@ -651,7 +668,8 @@ public:
 						// We need to make sure the package namespace is correct at this point
 						// Note: We don't test GIsEditor here as we need to mimic using compile-on-load what the compile during cook would have done when running with -game
 						{
-							const FString PackageNamespace = TextNamespaceUtil::GetPackageNamespace(Term->Source);
+							checkf(Term->SourcePin || Term->Source, TEXT("EmitTermExpr needs a valid source to correctly emit localized text"));
+							const FString PackageNamespace = TextNamespaceUtil::GetPackageNamespace(Term->SourcePin ? Term->SourcePin->GetOwningNode() : Term->Source);
 							Namespace = TextNamespaceUtil::BuildFullNamespace(Namespace, PackageNamespace);
 						}
 #endif // USE_STABLE_LOCALIZATION_KEYS
@@ -835,6 +853,7 @@ public:
 							Schema->ConvertPropertyToPinType(Prop, NewTerm.Type);
 							NewTerm.bIsLiteral = true;
 							NewTerm.Source = Term->Source;
+							NewTerm.SourcePin = Term->SourcePin;
 							Prop->ExportText_InContainer(ArrayIter, NewTerm.Name, StructData, StructData, NULL, PPF_None);
 							if (Prop->IsA(UTextProperty::StaticClass()))
 							{
@@ -872,6 +891,7 @@ public:
 					Schema->ConvertPropertyToPinType(InnerProp, NewTerm.Type);
 					NewTerm.bIsLiteral = true;
 					NewTerm.Source = Term->Source;
+					NewTerm.SourcePin = Term->SourcePin;
 					uint8* RawElemData = ScriptArrayHelper.GetRawPtr(ElemIdx);
 					InnerProp->ExportText_Direct(NewTerm.Name, RawElemData, RawElemData, NULL, PPF_None);
 					if (InnerProp->IsA(UTextProperty::StaticClass()))
diff --git a/Engine/Source/Editor/KismetCompiler/Public/EdGraphCompilerUtilities.h b/Engine/Source/Editor/KismetCompiler/Public/EdGraphCompilerUtilities.h
--- a/Engine/Source/Editor/KismetCompiler/Public/EdGraphCompilerUtilities.h
+++ b/Engine/Source/Editor/KismetCompiler/Public/EdGraphCompilerUtilities.h
@@ -1,9 +1,5 @@
 // Copyright 1998-2017 Epic Games, Inc. All Rights Reserved.
 
-
-#ifndef __EdGraphCompilerUtilities_h__
-#define __EdGraphCompilerUtilities_h__
-
 #pragma once
 
 #include "CoreMinimal.h"
@@ -93,5 +89,3 @@ protected:
 };
 
 //////////////////////////////////////////////////////////////////////////
-
-#endif	// __EdGraphCompilerUtilities_h__
diff --git a/Engine/Source/Editor/KismetWidgets/Private/SPinTypeSelector.cpp b/Engine/Source/Editor/KismetWidgets/Private/SPinTypeSelector.cpp
--- a/Engine/Source/Editor/KismetWidgets/Private/SPinTypeSelector.cpp
+++ b/Engine/Source/Editor/KismetWidgets/Private/SPinTypeSelector.cpp
@@ -45,7 +45,7 @@ int32 SDoubleImage::OnPaint(const FPaintArgs& Args, const FGeometry& AllottedGeo
 	SImage::OnPaint(Args, AllottedGeometry, MyClippingRect, OutDrawElements, LayerId, InWidgetStyle, bParentEnabled);
 
 	const bool bIsEnabled = ShouldBeEnabled(bParentEnabled);
-	const uint32 DrawEffects = bIsEnabled ? ESlateDrawEffect::None : ESlateDrawEffect::DisabledEffect;
+	const ESlateDrawEffect DrawEffects = bIsEnabled ? ESlateDrawEffect::None : ESlateDrawEffect::DisabledEffect;
 	// draw rest of the images, we reuse the LayerId because images are assumed to note overlap:
 	const FSlateBrush* SecondImageResolved = SecondImage.Get();
 	if (SecondImageResolved && SecondImageResolved->DrawAs != ESlateBrushDrawType::NoDrawType)
diff --git a/Engine/Source/Editor/LandscapeEditor/Private/LandscapeEdMode.cpp b/Engine/Source/Editor/LandscapeEditor/Private/LandscapeEdMode.cpp
--- a/Engine/Source/Editor/LandscapeEditor/Private/LandscapeEdMode.cpp
+++ b/Engine/Source/Editor/LandscapeEditor/Private/LandscapeEdMode.cpp
@@ -249,6 +249,7 @@ FEdModeLandscape::~FEdModeLandscape()
 	GSelectionRegionMaterial = NULL;
 	GMaskRegionMaterial = NULL;
 	GLandscapeBlackTexture = NULL;
+	GLandscapeLayerUsageMaterial = NULL;
 
 	InteractorPainting = nullptr;
 }
@@ -437,6 +438,8 @@ void FEdModeLandscape::Enter()
 	// Load UI settings from config file
 	UISettings->Load();
 
+	UpdateShownLayerList();
+
 	// Initialize current tool prior to creating the landscape toolkit in case it has a dependency on it
 	if (LandscapeList.Num() == 0)
 	{
@@ -782,6 +785,25 @@ bool FEdModeLandscape::MouseMove(FEditorViewportClient* ViewportClient, FViewpor
 	return Result;
 }
 
+bool FEdModeLandscape::GetCursor(EMouseCursor::Type& OutCursor) const
+{
+	if (!IsEditingEnabled())
+	{
+		return false;
+	}
+
+	bool Result = false;
+	if (NewLandscapePreviewMode == ENewLandscapePreviewMode::None)
+	{
+		if (CurrentTool)
+		{
+			Result = CurrentTool->GetCursor(OutCursor);
+		}
+	}
+
+	return Result;
+}
+
 bool FEdModeLandscape::DisallowMouseDeltaTracking() const
 {
 	// We never want to use the mouse delta tracker while painting
@@ -1418,18 +1440,22 @@ bool FEdModeLandscape::InputKey(FEditorViewportClient* ViewportClient, FViewport
 						else
 						{
 							Viewport->CaptureMouse(true);
-							bool bToolActive = CurrentTool->BeginTool(ViewportClient, CurrentToolTarget, HitLocation);
-							if (bToolActive)
-							{
-								ToolActiveViewport = Viewport;
-							}
-							else
+
+							if (CurrentTool->CanToolBeActivated())
 							{
-								ToolActiveViewport = nullptr;
-								Viewport->CaptureMouse(false);
+								bool bToolActive = CurrentTool->BeginTool(ViewportClient, CurrentToolTarget, HitLocation);
+								if (bToolActive)
+								{
+									ToolActiveViewport = Viewport;
+								}
+								else
+								{
+									ToolActiveViewport = nullptr;
+									Viewport->CaptureMouse(false);
+								}
+								ViewportClient->Invalidate(false, false);
+								return bToolActive;
 							}
-							ViewportClient->Invalidate(false, false);
-							return bToolActive;
 						}
 					}
 				}
@@ -1850,7 +1876,7 @@ void FEdModeLandscape::SetCurrentBrush(int32 BrushIndex)
 	}
 }
 
-const TArray<TSharedRef<FLandscapeTargetListInfo>>& FEdModeLandscape::GetTargetList()
+const TArray<TSharedRef<FLandscapeTargetListInfo>>& FEdModeLandscape::GetTargetList() const
 {
 	return LandscapeTargetList;
 }
@@ -1938,7 +1964,18 @@ int32 FEdModeLandscape::UpdateLandscapeList()
 			}
 			CurrentToolTarget.LandscapeInfo = LandscapeList[0].Info;
 			CurrentIndex = 0;
+
+			// Init UI to saved value
+			ALandscapeProxy* LandscapeProxy = CurrentToolTarget.LandscapeInfo->GetLandscapeProxy();
+
+			if (LandscapeProxy != nullptr)
+			{
+				UISettings->TargetDisplayOrder = LandscapeProxy->TargetDisplayOrder;
+			}
+
 			UpdateTargetList();
+			UpdateShownLayerList();
+
 			if (CurrentTool != nullptr)
 			{
 				CurrentTool->EnterTool();
@@ -1962,84 +1999,368 @@ void FEdModeLandscape::UpdateTargetList()
 {
 	LandscapeTargetList.Empty();
 
-	if (CurrentToolTarget.LandscapeInfo.IsValid() &&
-		CurrentToolTarget.LandscapeInfo->GetLandscapeProxy())
+	if (CurrentToolTarget.LandscapeInfo.IsValid())
 	{
-		CachedLandscapeMaterial = CurrentToolTarget.LandscapeInfo->GetLandscapeProxy()->GetLandscapeMaterial();
+		ALandscapeProxy* LandscapeProxy = CurrentToolTarget.LandscapeInfo->GetLandscapeProxy();
 
-		bool bFoundSelected = false;
+		if (LandscapeProxy != nullptr)
+		{
+			CachedLandscapeMaterial = LandscapeProxy->GetLandscapeMaterial();
 
-		// Add heightmap
-		LandscapeTargetList.Add(MakeShareable(new FLandscapeTargetListInfo(LOCTEXT("Heightmap", "Heightmap"), ELandscapeToolTargetType::Heightmap, CurrentToolTarget.LandscapeInfo.Get())));
+			bool bFoundSelected = false;
 
-		if (CurrentToolTarget.TargetType == ELandscapeToolTargetType::Heightmap)
-		{
-			bFoundSelected = true;
+			// Add heightmap
+			LandscapeTargetList.Add(MakeShareable(new FLandscapeTargetListInfo(LOCTEXT("Heightmap", "Heightmap"), ELandscapeToolTargetType::Heightmap, CurrentToolTarget.LandscapeInfo.Get())));
+
+			if (CurrentToolTarget.TargetType == ELandscapeToolTargetType::Heightmap)
+			{
+				bFoundSelected = true;
+			}
+
+			// Add visibility
+			FLandscapeInfoLayerSettings VisibilitySettings(ALandscapeProxy::VisibilityLayer, LandscapeProxy);
+			LandscapeTargetList.Add(MakeShareable(new FLandscapeTargetListInfo(LOCTEXT("Visibility", "Visibility"), ELandscapeToolTargetType::Visibility, VisibilitySettings)));
+
+			if (CurrentToolTarget.TargetType == ELandscapeToolTargetType::Visibility)
+			{
+				bFoundSelected = true;
+			}
+
+			// Add layers
+			UTexture2D* ThumbnailWeightmap = NULL;
+			UTexture2D* ThumbnailHeightmap = NULL;
+
+			TargetLayerStartingIndex = LandscapeTargetList.Num();
+
+			for (auto It = CurrentToolTarget.LandscapeInfo->Layers.CreateIterator(); It; It++)
+			{
+				FLandscapeInfoLayerSettings& LayerSettings = *It;
+
+				FName LayerName = LayerSettings.GetLayerName();
+
+				if (LayerSettings.LayerInfoObj == ALandscapeProxy::VisibilityLayer)
+				{
+					// Already handled above
+					continue;
+				}
+
+				if (!bFoundSelected &&
+					CurrentToolTarget.TargetType == ELandscapeToolTargetType::Weightmap &&
+					CurrentToolTarget.LayerInfo == LayerSettings.LayerInfoObj &&
+					CurrentToolTarget.LayerName == LayerSettings.LayerName)
+				{
+					bFoundSelected = true;
+				}
+
+				// Ensure thumbnails are up valid
+				if (LayerSettings.ThumbnailMIC == NULL)
+				{
+					if (ThumbnailWeightmap == NULL)
+					{
+						ThumbnailWeightmap = LoadObject<UTexture2D>(NULL, TEXT("/Engine/EditorLandscapeResources/LandscapeThumbnailWeightmap.LandscapeThumbnailWeightmap"), NULL, LOAD_None, NULL);
+					}
+					if (ThumbnailHeightmap == NULL)
+					{
+						ThumbnailHeightmap = LoadObject<UTexture2D>(NULL, TEXT("/Engine/EditorLandscapeResources/LandscapeThumbnailHeightmap.LandscapeThumbnailHeightmap"), NULL, LOAD_None, NULL);
+					}
+
+					// Construct Thumbnail MIC
+					UMaterialInterface* LandscapeMaterial = LayerSettings.Owner ? LayerSettings.Owner->GetLandscapeMaterial() : UMaterial::GetDefaultMaterial(MD_Surface);
+					LayerSettings.ThumbnailMIC = ALandscapeProxy::GetLayerThumbnailMIC(LandscapeMaterial, LayerName, ThumbnailWeightmap, ThumbnailHeightmap, LayerSettings.Owner);
+				}
+
+				// Add the layer
+				LandscapeTargetList.Add(MakeShareable(new FLandscapeTargetListInfo(FText::FromName(LayerName), ELandscapeToolTargetType::Weightmap, LayerSettings)));
+			}
+
+			if (!bFoundSelected)
+			{
+				CurrentToolTarget.TargetType = ELandscapeToolTargetType::Invalid;
+				CurrentToolTarget.LayerInfo = nullptr;
+				CurrentToolTarget.LayerName = NAME_None;
+			}
+
+			UpdateTargetLayerDisplayOrder(UISettings->TargetDisplayOrder);
 		}
+	}
 
-		// Add visibility
-		FLandscapeInfoLayerSettings VisibilitySettings(ALandscapeProxy::VisibilityLayer, CurrentToolTarget.LandscapeInfo->GetLandscapeProxy());
-		LandscapeTargetList.Add(MakeShareable(new FLandscapeTargetListInfo(LOCTEXT("Visibility", "Visibility"), ELandscapeToolTargetType::Visibility, VisibilitySettings)));
+	TargetsListUpdated.Broadcast();
+}
+
+void FEdModeLandscape::UpdateTargetLayerDisplayOrder(ELandscapeLayerDisplayMode InTargetDisplayOrder)
+{
+	if (!CurrentToolTarget.LandscapeInfo.IsValid())
+	{
+		return;
+	}
 
-		if (CurrentToolTarget.TargetType == ELandscapeToolTargetType::Visibility)
+	ALandscapeProxy* LandscapeProxy = CurrentToolTarget.LandscapeInfo->GetLandscapeProxy();
+
+	if (LandscapeProxy == nullptr)
+	{
+		return;
+	}
+
+	bool DetailPanelRefreshRequired = false;
+
+	// Save value to landscape
+	LandscapeProxy->TargetDisplayOrder = InTargetDisplayOrder;
+	TArray<FName>& SavedTargetNameList = LandscapeProxy->TargetDisplayOrderList;
+
+	switch (InTargetDisplayOrder)
+	{
+		case ELandscapeLayerDisplayMode::Default:
 		{
-			bFoundSelected = true;
-		}
+			SavedTargetNameList.Empty();
+
+			for (const TSharedRef<FLandscapeTargetListInfo>& TargetInfo : LandscapeTargetList)
+			{
+				SavedTargetNameList.Add(TargetInfo->LayerName);
+			}
 
-		// Add layers
-		UTexture2D* ThumbnailWeightmap = NULL;
-		UTexture2D* ThumbnailHeightmap = NULL;
+			DetailPanelRefreshRequired = true;
+		}
+		break;
 
-		for (auto It = CurrentToolTarget.LandscapeInfo->Layers.CreateIterator(); It; It++)
+		case ELandscapeLayerDisplayMode::Alphabetical:
 		{
-			FLandscapeInfoLayerSettings& LayerSettings = *It;
+			SavedTargetNameList.Empty();
+
+			// Add only layers to be able to sort them by name
+			for (int32 i = GetTargetLayerStartingIndex(); i < LandscapeTargetList.Num(); ++i)
+			{
+				SavedTargetNameList.Add(LandscapeTargetList[i]->LayerName);
+			}
 
-			FName LayerName = LayerSettings.GetLayerName();
+			SavedTargetNameList.Sort();
 
-			if (LayerSettings.LayerInfoObj == ALandscapeProxy::VisibilityLayer)
+			// Then insert the non layer target that shouldn't be sorted
+			for (int32 i = 0; i < GetTargetLayerStartingIndex(); ++i)
 			{
-				// Already handled above
-				continue;
+				SavedTargetNameList.Insert(LandscapeTargetList[i]->LayerName, i);
 			}
 
-			if (!bFoundSelected &&
-				CurrentToolTarget.TargetType == ELandscapeToolTargetType::Weightmap &&
-				CurrentToolTarget.LayerInfo == LayerSettings.LayerInfoObj &&
-				CurrentToolTarget.LayerName == LayerSettings.LayerName)
+			DetailPanelRefreshRequired = true;
+		}
+		break;
+
+		case ELandscapeLayerDisplayMode::UserSpecific:
+		{
+			for (const TSharedRef<FLandscapeTargetListInfo>& TargetInfo : LandscapeTargetList)
 			{
-				bFoundSelected = true;
+				bool Found = false;
+
+				for (const FName& LayerName : SavedTargetNameList)
+				{
+					if (TargetInfo->LayerName == LayerName)
+					{
+						Found = true;
+						break;
+					}
+				}
+
+				if (!Found)
+				{
+					DetailPanelRefreshRequired = true;
+					SavedTargetNameList.Add(TargetInfo->LayerName);
+				}
 			}
 
-			// Ensure thumbnails are up valid
-			if (LayerSettings.ThumbnailMIC == NULL)
+			// Handle the removing of elements from material
+			for (int32 i = SavedTargetNameList.Num() - 1; i >= 0; --i)
 			{
-				if (ThumbnailWeightmap == NULL)
+				bool Found = false;
+
+				for (const TSharedRef<FLandscapeTargetListInfo>& TargetInfo : LandscapeTargetList)
 				{
-					ThumbnailWeightmap = LoadObject<UTexture2D>(NULL, TEXT("/Engine/EditorLandscapeResources/LandscapeThumbnailWeightmap.LandscapeThumbnailWeightmap"), NULL, LOAD_None, NULL);
+					if (SavedTargetNameList[i] == TargetInfo->LayerName)
+					{
+						Found = true;
+						break;
+					}
 				}
-				if (ThumbnailHeightmap == NULL)
+
+				if (!Found)
 				{
-					ThumbnailHeightmap = LoadObject<UTexture2D>(NULL, TEXT("/Engine/EditorLandscapeResources/LandscapeThumbnailHeightmap.LandscapeThumbnailHeightmap"), NULL, LOAD_None, NULL);
+					DetailPanelRefreshRequired = true;
+					SavedTargetNameList.RemoveSingle(SavedTargetNameList[i]);
 				}
+			}
+		}
+		break;
+	}	
+
+	if (DetailPanelRefreshRequired)
+	{
+		if (Toolkit.IsValid())
+		{
+			StaticCastSharedPtr<FLandscapeToolKit>(Toolkit)->RefreshDetailPanel();
+		}
+	}
+}
+
+void FEdModeLandscape::UpdateShownLayerList()
+{
+	if (!CurrentToolTarget.LandscapeInfo.IsValid())
+	{
+		return;
+	}
+
+	// Make sure usage information is up to date
+	UpdateLayerUsageInformation();
+
+	bool DetailPanelRefreshRequired = false;
+
+	ShownTargetLayerList.Empty();
+
+	const TArray<FName>* DisplayOrderList = GetTargetDisplayOrderList();
 
-				// Construct Thumbnail MIC
-				UMaterialInterface* LandscapeMaterial = LayerSettings.Owner ? LayerSettings.Owner->GetLandscapeMaterial() : UMaterial::GetDefaultMaterial(MD_Surface);
-				LayerSettings.ThumbnailMIC = ALandscapeProxy::GetLayerThumbnailMIC(LandscapeMaterial, LayerName, ThumbnailWeightmap, ThumbnailHeightmap, LayerSettings.Owner);
+	if (DisplayOrderList == nullptr)
+	{
+		return;
+	}
+
+	for (const FName& LayerName : *DisplayOrderList)
+	{
+		for (const TSharedRef<FLandscapeTargetListInfo>& TargetInfo : GetTargetList())
+		{
+			if (TargetInfo->LayerName == LayerName)
+			{
+				// Keep a mapping of visible layer name to display order list so we can drag & drop proper items
+				if (ShouldShowLayer(TargetInfo))
+				{
+					ShownTargetLayerList.Add(TargetInfo->LayerName);
+					DetailPanelRefreshRequired = true;
+				}
+
+				break;
 			}
+		}
+	}	
+
+	if (DetailPanelRefreshRequired)
+	{
+		if (Toolkit.IsValid())
+		{
+			StaticCastSharedPtr<FLandscapeToolKit>(Toolkit)->RefreshDetailPanel();
+		}
+	}
+}
+
+void FEdModeLandscape::UpdateLayerUsageInformation()
+{
+	if (!CurrentToolTarget.LandscapeInfo.IsValid())
+	{
+		return;
+	}
+
+	bool DetailPanelRefreshRequired = false;
+
+	for (const TSharedRef<FLandscapeTargetListInfo>& TargetInfo : GetTargetList())
+	{		
+		TArray<ULandscapeComponent*> AllComponents;
+		CurrentToolTarget.LandscapeInfo->XYtoComponentMap.GenerateValueArray(AllComponents);
+		for (ULandscapeComponent* Component : AllComponents)
+		{
+			if (TargetInfo->LayerInfoObj.IsValid())
+			{
+				TArray<uint8> WeightmapTextureData;
+				FLandscapeComponentDataInterface DataInterface(Component);
+				DataInterface.GetWeightmapTextureData(TargetInfo->LayerInfoObj.Get(), WeightmapTextureData);
+
+				int32 UsageCount = 0;
+
+				for (uint8 Value : WeightmapTextureData)
+				{
+					UsageCount += Value;
+				}
+
+				bool PreviousValue = TargetInfo->LayerInfoObj->IsReferencedFromLoadedData;
+				TargetInfo->LayerInfoObj->IsReferencedFromLoadedData = UsageCount > 0;
 
-			// Add the layer
-			LandscapeTargetList.Add(MakeShareable(new FLandscapeTargetListInfo(FText::FromName(LayerName), ELandscapeToolTargetType::Weightmap, LayerSettings)));
+				if (PreviousValue != TargetInfo->LayerInfoObj->IsReferencedFromLoadedData)
+				{
+					DetailPanelRefreshRequired = true;
+				}
+
+				// Early exit as we already found a component using this layer
+				if (TargetInfo->LayerInfoObj->IsReferencedFromLoadedData)
+				{
+					break;
+				}
+			}
 		}
+	}
 
-		if (!bFoundSelected)
+	if (DetailPanelRefreshRequired)
+	{
+		if (Toolkit.IsValid())
 		{
-			CurrentToolTarget.TargetType = ELandscapeToolTargetType::Invalid;
-			CurrentToolTarget.LayerInfo = nullptr;
-			CurrentToolTarget.LayerName = NAME_None;
+			StaticCastSharedPtr<FLandscapeToolKit>(Toolkit)->RefreshDetailPanel();
 		}
 	}
+}
 
-	TargetsListUpdated.Broadcast();
+bool FEdModeLandscape::ShouldShowLayer(TSharedRef<FLandscapeTargetListInfo> Target) const
+{
+	if (!UISettings->ShowUnusedLayers)
+	{
+		return Target->LayerInfoObj.IsValid() && Target->LayerInfoObj.Get()->IsReferencedFromLoadedData;
+	}
+
+	return true;
+}
+
+const TArray<FName>& FEdModeLandscape::GetTargetShownList() const
+{
+	return ShownTargetLayerList;
+}
+
+int32 FEdModeLandscape::GetTargetLayerStartingIndex() const
+{
+	return TargetLayerStartingIndex;
+}
+
+const TArray<FName>* FEdModeLandscape::GetTargetDisplayOrderList() const
+{
+	if (!CurrentToolTarget.LandscapeInfo.IsValid())
+	{
+		return nullptr;
+	}
+
+	ALandscapeProxy* LandscapeProxy = CurrentToolTarget.LandscapeInfo->GetLandscapeProxy();
+
+	if (LandscapeProxy == nullptr)
+	{
+		return nullptr;
+	}
+
+	return &LandscapeProxy->TargetDisplayOrderList;
+}
+
+void FEdModeLandscape::MoveTargetLayerDisplayOrder(int32 IndexToMove, int32 IndexToDestination)
+{
+	if (!CurrentToolTarget.LandscapeInfo.IsValid())
+	{
+		return;
+	}
+
+	ALandscapeProxy* LandscapeProxy = CurrentToolTarget.LandscapeInfo->GetLandscapeProxy();
+
+	if (LandscapeProxy == nullptr)
+	{
+		return;
+	}
+	
+	FName Data = LandscapeProxy->TargetDisplayOrderList[IndexToMove];
+	LandscapeProxy->TargetDisplayOrderList.RemoveAt(IndexToMove);
+	LandscapeProxy->TargetDisplayOrderList.Insert(Data, IndexToDestination);
+
+	LandscapeProxy->TargetDisplayOrder = ELandscapeLayerDisplayMode::UserSpecific;
+	UISettings->TargetDisplayOrder = ELandscapeLayerDisplayMode::UserSpecific;
+
+	// Everytime we move something from the display order we must rebuild the shown layer list
+	UpdateShownLayerList();
 }
 
 FEdModeLandscape::FTargetsListUpdated FEdModeLandscape::TargetsListUpdated;
@@ -2050,6 +2371,7 @@ void FEdModeLandscape::HandleLevelsChanged()
 
 	UpdateLandscapeList();
 	UpdateTargetList();
+	UpdateShownLayerList();
 
 	// if the Landscape is deleted then close the landscape editor
 	if (bHadLandscape && CurrentToolTarget.LandscapeInfo == nullptr)
@@ -2074,6 +2396,7 @@ void FEdModeLandscape::OnMaterialCompilationFinished(UMaterialInterface* Materia
 	{
 		CurrentToolTarget.LandscapeInfo->UpdateLayerInfoMap();
 		UpdateTargetList();
+		UpdateShownLayerList();
 	}
 }
 
diff --git a/Engine/Source/Editor/LandscapeEditor/Private/LandscapeEdMode.h b/Engine/Source/Editor/LandscapeEditor/Private/LandscapeEdMode.h
--- a/Engine/Source/Editor/LandscapeEditor/Private/LandscapeEdMode.h
+++ b/Engine/Source/Editor/LandscapeEditor/Private/LandscapeEdMode.h
@@ -406,6 +406,8 @@ public:
 	virtual bool GetCustomDrawingCoordinateSystem(FMatrix& InMatrix, void* InData) override;
 	virtual bool GetCustomInputCoordinateSystem(FMatrix& InMatrix, void* InData) override;
 
+	virtual bool GetCursor(EMouseCursor::Type& OutCursor) const override;
+
 	/** Forces real-time perspective viewports */
 	void ForceRealTimeViewports(const bool bEnable, const bool bStoreCurrentState);
 
@@ -436,13 +438,27 @@ public:
 	void SetCurrentBrush(FName BrushName);
 	void SetCurrentBrush(int32 BrushIndex);
 
-	const TArray<TSharedRef<FLandscapeTargetListInfo>>& GetTargetList();
+	const TArray<TSharedRef<FLandscapeTargetListInfo>>& GetTargetList() const;
+	const TArray<FName>* GetTargetDisplayOrderList() const;
+	const TArray<FName>& GetTargetShownList() const;
+	int32 GetTargetLayerStartingIndex() const;
 	const TArray<FLandscapeListInfo>& GetLandscapeList();
 
 	void AddLayerInfo(ULandscapeLayerInfoObject* LayerInfo);
 
 	int32 UpdateLandscapeList();
 	void UpdateTargetList();
+	
+	/** Update Display order list */
+	void UpdateTargetLayerDisplayOrder(ELandscapeLayerDisplayMode InTargetDisplayOrder);
+	void MoveTargetLayerDisplayOrder(int32 IndexToMove, int32 IndexToDestination);
+
+	/** Update shown layer list */	
+	void UpdateShownLayerList();
+	bool ShouldShowLayer(TSharedRef<FLandscapeTargetListInfo> Target) const;
+	void UpdateLayerUsageInformation();
+
+	void RefreshDetailPanel();
 
 	DECLARE_EVENT(FEdModeLandscape, FTargetsListUpdated);
 	static FTargetsListUpdated TargetsListUpdated;
@@ -483,6 +499,10 @@ public:
 private:
 	TArray<TSharedRef<FLandscapeTargetListInfo>> LandscapeTargetList;
 	TArray<FLandscapeListInfo> LandscapeList;
+	TArray<FName> ShownTargetLayerList;
+	
+	/** Represent the index offset of the target layer in LandscapeTargetList */
+	int32 TargetLayerStartingIndex;
 
 	UMaterialInterface* CachedLandscapeMaterial;
 
diff --git a/Engine/Source/Editor/LandscapeEditor/Private/LandscapeEdModeMirrorTool.cpp b/Engine/Source/Editor/LandscapeEditor/Private/LandscapeEdModeMirrorTool.cpp
--- a/Engine/Source/Editor/LandscapeEditor/Private/LandscapeEdModeMirrorTool.cpp
+++ b/Engine/Source/Editor/LandscapeEditor/Private/LandscapeEdModeMirrorTool.cpp
@@ -631,6 +631,7 @@ public:
 				LandscapeEdit.GetWeightData(LayerInfo, TempMinX, TempMinY, TempMaxX, TempMaxY, &SourceWeightData[0], SourceSizeX);
 				ApplyMirrorInternal<uint8>(SourceWeightData, DestWeightData, SourceSizeX, SourceSizeY, DestSizeX, DestSizeY, MirrorPos, BlendWidth);
 				LandscapeEdit.SetAlphaData(LayerInfo, DestMinX, DestMinY, DestMaxX, DestMaxY, &DestWeightData[0], DestSizeX, ELandscapeLayerPaintingRestriction::None, false, false);
+				//LayerInfo->IsReferencedFromLoadedData = true;
 			}
 		}
 
@@ -657,6 +658,8 @@ public:
 
 			// Flush dynamic foliage (grass)
 			ALandscapeProxy::InvalidateGeneratedComponentData(Components);
+
+			EdMode->UpdateLayerUsageInformation();
 		}
 	}
 
diff --git a/Engine/Source/Editor/LandscapeEditor/Private/LandscapeEdModePaintTools.cpp b/Engine/Source/Editor/LandscapeEditor/Private/LandscapeEdModePaintTools.cpp
--- a/Engine/Source/Editor/LandscapeEditor/Private/LandscapeEdModePaintTools.cpp
+++ b/Engine/Source/Editor/LandscapeEditor/Private/LandscapeEdModePaintTools.cpp
@@ -733,18 +733,25 @@ public:
 						case ELandscapeToolNoiseMode::Add:
 							if (Delta < 0)
 							{
-								DataScanline[X] = FMath::Lerp(DataScanline[X], FlattenHeight, Strength);
+								DataScanline[X] = FMath::CeilToInt(FMath::Lerp((float)DataScanline[X], (float)FlattenHeight, Strength));
 							}
 							break;
 						case ELandscapeToolNoiseMode::Sub:
 							if (Delta > 0)
 							{
-								DataScanline[X] = FMath::Lerp(DataScanline[X], FlattenHeight, Strength);
+								DataScanline[X] = FMath::FloorToInt(FMath::Lerp((float)DataScanline[X], (float)FlattenHeight, Strength));
 							}
 							break;
 						default:
 						case ELandscapeToolNoiseMode::Both:
-							DataScanline[X] = FMath::Lerp(DataScanline[X], FlattenHeight, Strength);
+							if (Delta > 0)
+							{
+								DataScanline[X] = FMath::FloorToInt(FMath::Lerp((float)DataScanline[X], (float)FlattenHeight, Strength));
+							}
+							else
+							{
+								DataScanline[X] = FMath::CeilToInt(FMath::Lerp((float)DataScanline[X], (float)FlattenHeight, Strength));
+							}
 							break;
 						}
 					}
@@ -759,18 +766,25 @@ public:
 						case ELandscapeToolNoiseMode::Add:
 							if (PlaneDist < 0)
 							{
-								DataScanline[X] = FMath::Lerp(DataScanline[X], DestValue, Strength);
+								DataScanline[X] = FMath::CeilToInt(FMath::Lerp((float)DataScanline[X], (float)DestValue, Strength));
 							}
 							break;
 						case ELandscapeToolNoiseMode::Sub:
 							if (PlaneDist > 0)
 							{
-								DataScanline[X] = FMath::Lerp(DataScanline[X], DestValue, Strength);
+								DataScanline[X] = FMath::FloorToInt(FMath::Lerp((float)DataScanline[X], (float)DestValue, Strength));
 							}
 							break;
 						default:
 						case ELandscapeToolNoiseMode::Both:
-							DataScanline[X] = FMath::Lerp(DataScanline[X], DestValue, Strength);
+							if (PlaneDist > 0)
+							{
+								DataScanline[X] = FMath::FloorToInt(FMath::Lerp((float)DataScanline[X], (float)DestValue, Strength));
+							}
+							else
+							{
+								DataScanline[X] = FMath::CeilToInt(FMath::Lerp((float)DataScanline[X], (float)DestValue, Strength));
+							}
 							break;
 						}
 					}
@@ -789,16 +803,39 @@ class FLandscapeToolFlatten : public FLandscapeToolPaintBase < ToolTarget, FLand
 protected:
 	UStaticMesh* PlaneMesh;
 	UStaticMeshComponent* MeshComponent;
+	bool CanToolBeActivatedNextTick;
+	bool CanToolBeActivatedValue;
+	float EyeDropperFlattenTargetValue;
 
 public:
 	FLandscapeToolFlatten(FEdModeLandscape* InEdMode)
 		: FLandscapeToolPaintBase<ToolTarget, FLandscapeToolStrokeFlatten<ToolTarget>>(InEdMode)
 		, PlaneMesh(LoadObject<UStaticMesh>(NULL, TEXT("/Engine/EditorLandscapeResources/FlattenPlaneMesh.FlattenPlaneMesh")))
 		, MeshComponent(NULL)
+		, CanToolBeActivatedNextTick(false)
+		, CanToolBeActivatedValue(false)
+		, EyeDropperFlattenTargetValue(0.0f)
 	{
 		check(PlaneMesh);
 	}
 
+	virtual bool GetCursor(EMouseCursor::Type& OutCursor) const override
+	{ 
+		if (this->EdMode->UISettings->bFlattenEyeDropperModeActivated)
+		{
+			OutCursor = EMouseCursor::EyeDropper;
+			return true;
+		}
+		
+		return false; 
+	}
+
+	virtual void SetCanToolBeActivated(bool Value) override
+	{ 
+		CanToolBeActivatedNextTick = true;
+		CanToolBeActivatedValue = Value;
+	}
+
 	virtual void AddReferencedObjects(FReferenceCollector& Collector) override
 	{
 		Collector.AddReferencedObject(PlaneMesh);
@@ -810,6 +847,12 @@ public:
 
 	virtual void Tick(FEditorViewportClient* ViewportClient, float DeltaTime) override
 	{
+		if (CanToolBeActivatedNextTick)
+		{
+			this->bCanToolBeActivated = CanToolBeActivatedValue;
+			CanToolBeActivatedNextTick = false;
+		}
+
 		FLandscapeToolPaintBase<ToolTarget, FLandscapeToolStrokeFlatten<ToolTarget>>::Tick(ViewportClient, DeltaTime);
 
 		bool bShowGrid = this->EdMode->UISettings->bUseFlattenTarget && this->EdMode->CurrentToolTarget.TargetType == ELandscapeToolTargetType::Heightmap && this->EdMode->UISettings->bShowFlattenTargetPreview;
@@ -831,6 +874,13 @@ public:
 			Origin.Y = FMath::RoundToFloat(MousePosition.Y);
 			Origin.Z = (FMath::RoundToFloat((this->EdMode->UISettings->FlattenTarget - LocalToWorld.GetTranslation().Z) / LocalToWorld.GetScale3D().Z * LANDSCAPE_INV_ZSCALE) - 0.1f) * LANDSCAPE_ZSCALE;
 			MeshComponent->SetRelativeLocation(Origin, false);
+
+			// Clamp the value to the height map
+			uint16 TexHeight = LandscapeDataAccess::GetTexHeight(MousePosition.Z);
+			float Height = LandscapeDataAccess::GetLocalHeight(TexHeight);
+
+			// Convert the height back to world space
+			this->EdMode->UISettings->FlattenEyeDropperModeDesiredTarget = (Height * LocalToWorld.GetScale3D().Z) + LocalToWorld.GetTranslation().Z;
 		}
 
 		return bResult;
diff --git a/Engine/Source/Editor/LandscapeEditor/Private/LandscapeEdModeTools.h b/Engine/Source/Editor/LandscapeEditor/Private/LandscapeEdModeTools.h
--- a/Engine/Source/Editor/LandscapeEditor/Private/LandscapeEdModeTools.h
+++ b/Engine/Source/Editor/LandscapeEditor/Private/LandscapeEdModeTools.h
@@ -964,6 +964,7 @@ struct FAlphamapAccessor
 			ALandscapeProxy::InvalidateGeneratedComponentData(Components);
 
 			LandscapeEdit.SetAlphaData(LayerInfo, X1, Y1, X2, Y2, Data, 0, PaintingRestriction, bBlendWeight, bUseTotalNormalize);
+			//LayerInfo->IsReferencedFromLoadedData = true;
 			ModifiedComponents.Append(Components);
 		}
 	}
@@ -1311,6 +1312,7 @@ public:
 	FLandscapeToolBase(FEdModeLandscape* InEdMode)
 		: EdMode(InEdMode)
 		, bToolActive(false)
+		, bCanToolBeActivated(true)
 	{
 	}
 
@@ -1375,6 +1377,7 @@ public:
 		ToolStroke.Reset();
 		bToolActive = false;
 		EdMode->CurrentBrush->EndStroke();
+		EdMode->UpdateLayerUsageInformation();
 	}
 
 	virtual bool MouseMove(FEditorViewportClient* ViewportClient, FViewport* Viewport, int32 x, int32 y) override
@@ -1405,12 +1408,16 @@ public:
 
 	virtual bool IsToolActive() const override { return bToolActive;  }
 
+	virtual void SetCanToolBeActivated(bool Value) { bCanToolBeActivated = Value; }
+	virtual bool CanToolBeActivated() const {	return bCanToolBeActivated; }
+
 protected:
 	TArray<FLandscapeToolInteractorPosition> InteractorPositions;
 	FVector2D LastInteractorPosition;
 	float TimeSinceLastInteractorMove;
 	FEdModeLandscape* EdMode;
 	bool bToolActive;
+	bool bCanToolBeActivated;
 	TOptional<TStrokeClass> ToolStroke;
 
 	bool IsModifierPressed(const class FEditorViewportClient* ViewportClient, const UViewportInteractor* Interactor = nullptr)
@@ -1425,3 +1432,14 @@ protected:
 		return IsShiftDown(ViewportClient->Viewport) || bIsInteractorModifierPressed;
 	}
 };
+
+struct FToolFlattenCustomData
+{
+	FToolFlattenCustomData()
+		: ActiveEyeDropperMode(false)
+		, EyeDropperModeHeight(0.0f)
+	{}
+
+	bool ActiveEyeDropperMode;
+	float EyeDropperModeHeight;
+};
diff --git a/Engine/Source/Editor/LandscapeEditor/Private/LandscapeEditorDetailCustomization_Base.cpp b/Engine/Source/Editor/LandscapeEditor/Private/LandscapeEditorDetailCustomization_Base.cpp
--- a/Engine/Source/Editor/LandscapeEditor/Private/LandscapeEditorDetailCustomization_Base.cpp
+++ b/Engine/Source/Editor/LandscapeEditor/Private/LandscapeEditorDetailCustomization_Base.cpp
@@ -36,29 +36,6 @@ bool FLandscapeEditorDetailCustomization_Base::IsBrushSetActive(FName BrushSetNa
 	return false;
 }
 
-TOptional<int32> FLandscapeEditorDetailCustomization_Base::OnGetValue(TSharedRef<IPropertyHandle> PropertyHandle)
-{
-	int32 IntValue = 0;
-	if (ensure(PropertyHandle->GetValue(IntValue) == FPropertyAccess::Success))
-	{
-		return TOptional<int32>(IntValue);
-	}
-
-	// Value couldn't be accessed. Return an unset value
-	return TOptional<int32>();
-}
-
-void FLandscapeEditorDetailCustomization_Base::OnValueChanged(int32 NewValue, TSharedRef<IPropertyHandle> PropertyHandle)
-{
-	const EPropertyValueSetFlags::Type Flags = EPropertyValueSetFlags::InteractiveChange;
-	ensure(PropertyHandle->SetValue(NewValue, Flags) == FPropertyAccess::Success);
-}
-
-void FLandscapeEditorDetailCustomization_Base::OnValueCommitted(int32 NewValue, ETextCommit::Type CommitType, TSharedRef<IPropertyHandle> PropertyHandle)
-{
-	ensure(PropertyHandle->SetValue(NewValue) == FPropertyAccess::Success);
-}
-
 //////////////////////////////////////////////////////////////////////////
 
 FEdModeLandscape* FLandscapeEditorStructCustomization_Base::GetEditorMode()
diff --git a/Engine/Source/Editor/LandscapeEditor/Private/LandscapeEditorDetailCustomization_Base.h b/Engine/Source/Editor/LandscapeEditor/Private/LandscapeEditorDetailCustomization_Base.h
--- a/Engine/Source/Editor/LandscapeEditor/Private/LandscapeEditorDetailCustomization_Base.h
+++ b/Engine/Source/Editor/LandscapeEditor/Private/LandscapeEditorDetailCustomization_Base.h
@@ -25,9 +25,14 @@ protected:
 	static bool IsToolActive(FName ToolName);
 	static bool IsBrushSetActive(FName BrushSetName);
 
-	static TOptional<int32> OnGetValue(TSharedRef<IPropertyHandle> PropertyHandle);
-	static void OnValueChanged(int32 NewValue, TSharedRef<IPropertyHandle> PropertyHandle);
-	static void OnValueCommitted(int32 NewValue, ETextCommit::Type CommitType, TSharedRef<IPropertyHandle> PropertyHandle);
+	template<typename type>
+	static TOptional<type> OnGetValue(TSharedRef<IPropertyHandle> PropertyHandle);
+	
+	template<typename type>
+	static void OnValueChanged(type NewValue, TSharedRef<IPropertyHandle> PropertyHandle);
+	
+	template<typename type>
+	static void OnValueCommitted(type NewValue, ETextCommit::Type CommitType, TSharedRef<IPropertyHandle> PropertyHandle);
 
 	template<typename type>
 	static type GetPropertyValue(TSharedRef<IPropertyHandle> PropertyHandle);
@@ -44,6 +49,32 @@ protected:
 	static void SetPropertyValue(type NewValue, ETextCommit::Type CommitInfo, TSharedRef<IPropertyHandle> PropertyHandle);
 };
 
+template<typename type>
+TOptional<type> FLandscapeEditorDetailCustomization_Base::OnGetValue(TSharedRef<IPropertyHandle> PropertyHandle)
+{
+	type Value;
+	if (ensure(PropertyHandle->GetValue(Value) == FPropertyAccess::Success))
+	{
+		return TOptional<type>(Value);
+	}
+
+	// Value couldn't be accessed. Return an unset value
+	return TOptional<type>();
+}
+
+template<typename type>
+void FLandscapeEditorDetailCustomization_Base::OnValueChanged(type NewValue, TSharedRef<IPropertyHandle> PropertyHandle)
+{
+	const EPropertyValueSetFlags::Type Flags = EPropertyValueSetFlags::InteractiveChange;
+	ensure(PropertyHandle->SetValue(NewValue, Flags) == FPropertyAccess::Success);
+}
+
+template<typename type>
+void FLandscapeEditorDetailCustomization_Base::OnValueCommitted(type NewValue, ETextCommit::Type CommitType, TSharedRef<IPropertyHandle> PropertyHandle)
+{
+	ensure(PropertyHandle->SetValue(NewValue) == FPropertyAccess::Success);
+}
+
 template<typename type>
 type FLandscapeEditorDetailCustomization_Base::GetPropertyValue(TSharedRef<IPropertyHandle> PropertyHandle)
 {
diff --git a/Engine/Source/Editor/LandscapeEditor/Private/LandscapeEditorDetailCustomization_CopyPaste.cpp b/Engine/Source/Editor/LandscapeEditor/Private/LandscapeEditorDetailCustomization_CopyPaste.cpp
--- a/Engine/Source/Editor/LandscapeEditor/Private/LandscapeEditorDetailCustomization_CopyPaste.cpp
+++ b/Engine/Source/Editor/LandscapeEditor/Private/LandscapeEditorDetailCustomization_CopyPaste.cpp
@@ -130,9 +130,9 @@ void FLandscapeEditorDetailCustomization_CopyPaste::CustomizeDetails(IDetailLayo
 			.MaxSliderValue(8192)
 			.AllowSpin(true)
 			.UndeterminedString(NSLOCTEXT("PropertyEditor", "MultipleValues", "Multiple Values"))
-			.Value_Static(&FLandscapeEditorDetailCustomization_Base::OnGetValue, PropertyHandle_ImportSize_X)
-			.OnValueChanged_Static(&FLandscapeEditorDetailCustomization_Base::OnValueChanged, PropertyHandle_ImportSize_X)
-			.OnValueCommitted_Static(&FLandscapeEditorDetailCustomization_Base::OnValueCommitted, PropertyHandle_ImportSize_X)
+			.Value_Static(&FLandscapeEditorDetailCustomization_Base::OnGetValue<int32>, PropertyHandle_ImportSize_X)
+			.OnValueChanged_Static(&FLandscapeEditorDetailCustomization_Base::OnValueChanged<int32>, PropertyHandle_ImportSize_X)
+			.OnValueCommitted_Static(&FLandscapeEditorDetailCustomization_Base::OnValueCommitted<int32>, PropertyHandle_ImportSize_X)
 		]
 		+ SHorizontalBox::Slot()
 		.AutoWidth()
@@ -155,9 +155,9 @@ void FLandscapeEditorDetailCustomization_CopyPaste::CustomizeDetails(IDetailLayo
 			.MaxSliderValue(8192)
 			.AllowSpin(true)
 			.UndeterminedString(NSLOCTEXT("PropertyEditor", "MultipleValues", "Multiple Values"))
-			.Value_Static(&FLandscapeEditorDetailCustomization_Base::OnGetValue, PropertyHandle_ImportSize_Y)
-			.OnValueChanged_Static(&FLandscapeEditorDetailCustomization_Base::OnValueChanged, PropertyHandle_ImportSize_Y)
-			.OnValueCommitted_Static(&FLandscapeEditorDetailCustomization_Base::OnValueCommitted, PropertyHandle_ImportSize_Y)
+			.Value_Static(&FLandscapeEditorDetailCustomization_Base::OnGetValue<int32>, PropertyHandle_ImportSize_Y)
+			.OnValueChanged_Static(&FLandscapeEditorDetailCustomization_Base::OnValueChanged<int32>, PropertyHandle_ImportSize_Y)
+			.OnValueCommitted_Static(&FLandscapeEditorDetailCustomization_Base::OnValueCommitted<int32>, PropertyHandle_ImportSize_Y)
 		]
 	];
 
diff --git a/Engine/Source/Editor/LandscapeEditor/Private/LandscapeEditorDetailCustomization_MiscTools.cpp b/Engine/Source/Editor/LandscapeEditor/Private/LandscapeEditorDetailCustomization_MiscTools.cpp
--- a/Engine/Source/Editor/LandscapeEditor/Private/LandscapeEditorDetailCustomization_MiscTools.cpp
+++ b/Engine/Source/Editor/LandscapeEditor/Private/LandscapeEditorDetailCustomization_MiscTools.cpp
@@ -13,6 +13,10 @@
 #include "DetailCategoryBuilder.h"
 
 #include "ScopedTransaction.h"
+#include "IDetailPropertyRow.h"
+#include "SNumericEntryBox.h"
+#include "SFlattenHeightEyeDropperButton.h"
+#include "LandscapeEdModeTools.h"
 
 #define LOCTEXT_NAMESPACE "LandscapeEditor.Tools"
 
@@ -50,6 +54,53 @@ void FLandscapeEditorDetailCustomization_MiscTools::CustomizeDetails(IDetailLayo
 		];
 	}
 
+	if (IsToolActive("Flatten"))
+	{
+		TSharedRef<IPropertyHandle> FlattenValueProperty = DetailBuilder.GetProperty(GET_MEMBER_NAME_CHECKED(ULandscapeEditorObject, FlattenTarget));
+		IDetailPropertyRow& FlattenValueRow = ToolsCategory.AddProperty(FlattenValueProperty);
+		FlattenValueRow.CustomWidget()
+			.NameContent()
+			[
+				FlattenValueProperty->CreatePropertyNameWidget()
+			]
+			.ValueContent()
+			[
+				SNew(SHorizontalBox)
+				+SHorizontalBox::Slot()
+				.Padding(0.0f, 2.0f, 5.0f, 2.0f)
+				.FillWidth(1.0f)
+				.HAlign(HAlign_Left)
+				.VAlign(VAlign_Center)
+				[
+					SNew(SNumericEntryBox<float>)
+					.AllowSpin(true)
+					.Font(DetailBuilder.GetDetailFont())
+					.Value(this, &FLandscapeEditorDetailCustomization_MiscTools::GetFlattenValue)
+					.OnValueChanged_Static(&FLandscapeEditorDetailCustomization_Base::OnValueChanged<float>, FlattenValueProperty)
+					.OnValueCommitted_Static(&FLandscapeEditorDetailCustomization_Base::OnValueCommitted<float>, FlattenValueProperty)
+					.MinValue(-32768.0f)
+					.MaxValue(32768.0f)
+					.SliderExponentNeutralValue(0.0f)
+					.SliderExponent(5.0f)
+					.ShiftMouseMovePixelPerDelta(20)
+					.MinSliderValue(-32768.0f)
+					.MaxSliderValue(32768.0f)
+					.MinDesiredValueWidth(75.0f)
+					.ToolTipText(LOCTEXT("FlattenToolTips", "Target height to flatten towards (in Unreal Units)"))
+				]
+				+SHorizontalBox::Slot()
+				.AutoWidth()
+				.Padding(0.0f, 2.0f, 5.0f, 2.0f)
+				.HAlign(HAlign_Right)
+				.VAlign(VAlign_Center)
+				[
+					SNew(SFlattenHeightEyeDropperButton)
+					.OnBegin(this, &FLandscapeEditorDetailCustomization_MiscTools::OnBeginFlattenToolEyeDrop)
+					.OnComplete(this, &FLandscapeEditorDetailCustomization_MiscTools::OnCompletedFlattenToolEyeDrop)
+				]		
+			];
+	}
+
 	if (IsToolActive("Splines"))
 	{
 		ToolsCategory.AddCustomRow(LOCTEXT("ApplySplinesLabel", "Apply Splines"))
@@ -345,4 +396,45 @@ FReply FLandscapeEditorDetailCustomization_MiscTools::OnResetMirrorPointButtonCl
 	return FReply::Handled();
 }
 
+TOptional<float> FLandscapeEditorDetailCustomization_MiscTools::GetFlattenValue() const
+{
+	FEdModeLandscape* LandscapeEdMode = GetEditorMode();
+	if (LandscapeEdMode != nullptr && IsToolActive(FName("Flatten")))
+	{
+		if (LandscapeEdMode->UISettings->bFlattenEyeDropperModeActivated)
+		{
+			return LandscapeEdMode->UISettings->FlattenEyeDropperModeDesiredTarget;
+		}
+
+		return LandscapeEdMode->UISettings->FlattenTarget;
+	}
+
+	return 0.0f;
+}
+
+void FLandscapeEditorDetailCustomization_MiscTools::OnBeginFlattenToolEyeDrop()
+{
+	FEdModeLandscape* LandscapeEdMode = GetEditorMode();
+	if (LandscapeEdMode != nullptr && IsToolActive(FName("Flatten")))
+	{
+		LandscapeEdMode->UISettings->bFlattenEyeDropperModeActivated = true;
+		LandscapeEdMode->CurrentTool->SetCanToolBeActivated(false);
+	}
+}
+
+void FLandscapeEditorDetailCustomization_MiscTools::OnCompletedFlattenToolEyeDrop(bool Canceled)
+{
+	FEdModeLandscape* LandscapeEdMode = GetEditorMode();
+	if (LandscapeEdMode != nullptr && IsToolActive(FName("Flatten")))
+	{
+		LandscapeEdMode->UISettings->bFlattenEyeDropperModeActivated = false;
+		LandscapeEdMode->CurrentTool->SetCanToolBeActivated(true);
+
+		if (!Canceled)
+		{
+			LandscapeEdMode->UISettings->FlattenTarget = LandscapeEdMode->UISettings->FlattenEyeDropperModeDesiredTarget;
+		}
+	}
+}
+
 #undef LOCTEXT_NAMESPACE
diff --git a/Engine/Source/Editor/LandscapeEditor/Private/LandscapeEditorDetailCustomization_MiscTools.h b/Engine/Source/Editor/LandscapeEditor/Private/LandscapeEditorDetailCustomization_MiscTools.h
--- a/Engine/Source/Editor/LandscapeEditor/Private/LandscapeEditorDetailCustomization_MiscTools.h
+++ b/Engine/Source/Editor/LandscapeEditor/Private/LandscapeEditorDetailCustomization_MiscTools.h
@@ -46,4 +46,11 @@ protected:
 	// Mirror Tool
 	static FReply OnApplyMirrorButtonClicked();
 	static FReply OnResetMirrorPointButtonClicked();
+
+	// Flatten Tool
+	void HandleFlattenValueChanged(float NewValue);
+	void OnBeginFlattenToolEyeDrop();
+	void OnCompletedFlattenToolEyeDrop(bool Canceled);
+	TOptional<float> GetFlattenValue() const;
+
 };
diff --git a/Engine/Source/Editor/LandscapeEditor/Private/LandscapeEditorDetailCustomization_NewLandscape.cpp b/Engine/Source/Editor/LandscapeEditor/Private/LandscapeEditorDetailCustomization_NewLandscape.cpp
--- a/Engine/Source/Editor/LandscapeEditor/Private/LandscapeEditorDetailCustomization_NewLandscape.cpp
+++ b/Engine/Source/Editor/LandscapeEditor/Private/LandscapeEditorDetailCustomization_NewLandscape.cpp
@@ -332,9 +332,9 @@ void FLandscapeEditorDetailCustomization_NewLandscape::CustomizeDetails(IDetailL
 			.MaxSliderValue(32)
 			.AllowSpin(true)
 			.UndeterminedString(NSLOCTEXT("PropertyEditor", "MultipleValues", "Multiple Values"))
-			.Value_Static(&FLandscapeEditorDetailCustomization_Base::OnGetValue, PropertyHandle_ComponentCount_X)
-			.OnValueChanged_Static(&FLandscapeEditorDetailCustomization_Base::OnValueChanged, PropertyHandle_ComponentCount_X)
-			.OnValueCommitted_Static(&FLandscapeEditorDetailCustomization_Base::OnValueCommitted, PropertyHandle_ComponentCount_X)
+			.Value_Static(&FLandscapeEditorDetailCustomization_Base::OnGetValue<int32>, PropertyHandle_ComponentCount_X)
+			.OnValueChanged_Static(&FLandscapeEditorDetailCustomization_Base::OnValueChanged<int32>, PropertyHandle_ComponentCount_X)
+			.OnValueCommitted_Static(&FLandscapeEditorDetailCustomization_Base::OnValueCommitted<int32>, PropertyHandle_ComponentCount_X)
 		]
 		+ SHorizontalBox::Slot()
 		.AutoWidth()
@@ -357,9 +357,9 @@ void FLandscapeEditorDetailCustomization_NewLandscape::CustomizeDetails(IDetailL
 			.MaxSliderValue(32)
 			.AllowSpin(true)
 			.UndeterminedString(NSLOCTEXT("PropertyEditor", "MultipleValues", "Multiple Values"))
-			.Value_Static(&FLandscapeEditorDetailCustomization_Base::OnGetValue, PropertyHandle_ComponentCount_Y)
-			.OnValueChanged_Static(&FLandscapeEditorDetailCustomization_Base::OnValueChanged, PropertyHandle_ComponentCount_Y)
-			.OnValueCommitted_Static(&FLandscapeEditorDetailCustomization_Base::OnValueCommitted, PropertyHandle_ComponentCount_Y)
+			.Value_Static(&FLandscapeEditorDetailCustomization_Base::OnGetValue<int32>, PropertyHandle_ComponentCount_Y)
+			.OnValueChanged_Static(&FLandscapeEditorDetailCustomization_Base::OnValueChanged<int32>, PropertyHandle_ComponentCount_Y)
+			.OnValueCommitted_Static(&FLandscapeEditorDetailCustomization_Base::OnValueCommitted<int32>, PropertyHandle_ComponentCount_Y)
 		]
 	];
 
