Allow unloaded String Table assets to be picked by text properties

diff --git a/Engine/Source/Editor/EditorWidgets/Private/STextPropertyEditableTextBox.cpp b/Engine/Source/Editor/EditorWidgets/Private/STextPropertyEditableTextBox.cpp
--- a/Engine/Source/Editor/EditorWidgets/Private/STextPropertyEditableTextBox.cpp
+++ b/Engine/Source/Editor/EditorWidgets/Private/STextPropertyEditableTextBox.cpp
@@ -340,7 +340,7 @@ void STextPropertyEditableStringTableReference::UpdateStringTableComboOptions()
 
 		for (const FAssetData& StringTableAsset : StringTableAssets)
 		{
-			// Only allow string tables assets that have entries to be visible otherwise unexpected behaviour happens for the user
+			// Only allow string tables assets that have entries to be visible otherwise unexpected behavior happens for the user
 			bool HasEntries = false;
 			FStringTableConstPtr StringTable = FStringTableRegistry::Get().FindStringTable(StringTableAsset.ObjectPath);
 			if (StringTable.IsValid())
@@ -351,6 +351,11 @@ void STextPropertyEditableStringTableReference::UpdateStringTableComboOptions()
 					return false; // stop enumeration
 				});
 			}
+			else
+			{
+				// Asset is currently unloaded, so just assume it has entries
+				HasEntries = true;
+			}
 
 			if (!HasEntries)
 			{
