Unshelved from dev-niagara:

Fix SDropTarget not using horizontal and vertical image brush args

diff --git a/Engine/Source/Editor/EditorWidgets/Private/SDropTarget.cpp b/Engine/Source/Editor/EditorWidgets/Private/SDropTarget.cpp
--- a/Engine/Source/Editor/EditorWidgets/Private/SDropTarget.cpp
+++ b/Engine/Source/Editor/EditorWidgets/Private/SDropTarget.cpp
@@ -153,17 +153,14 @@ int32 SDropTarget::OnPaint(const FPaintArgs& Args, const FGeometry& AllottedGeom
 		{
 			FLinearColor DashColor = bAllowDrop ? ValidColor : InvalidColor;
 
-			const FSlateBrush* HorizontalBrush = FEditorStyle::GetBrush("WideDash.Horizontal");
-			const FSlateBrush* VerticalBrush = FEditorStyle::GetBrush("WideDash.Vertical");
-
 			int32 DashLayer = LayerId + 1;
 
 			// Top
 			FSlateDrawElement::MakeBox(
 				OutDrawElements,
 				DashLayer,
-				AllottedGeometry.ToPaintGeometry(FVector2D(0, 0), FVector2D(AllottedGeometry.GetLocalSize().X, HorizontalBrush->ImageSize.Y)),
-				HorizontalBrush,
+				AllottedGeometry.ToPaintGeometry(FVector2D(0, 0), FVector2D(AllottedGeometry.GetLocalSize().X, HorizontalImage->ImageSize.Y)),
+				HorizontalImage,
 				ESlateDrawEffect::None,
 				DashColor);
 
@@ -171,8 +168,8 @@ int32 SDropTarget::OnPaint(const FPaintArgs& Args, const FGeometry& AllottedGeom
 			FSlateDrawElement::MakeBox(
 				OutDrawElements,
 				DashLayer,
-				AllottedGeometry.ToPaintGeometry(FVector2D(0, AllottedGeometry.GetLocalSize().Y - HorizontalBrush->ImageSize.Y), FVector2D(AllottedGeometry.Size.X, HorizontalBrush->ImageSize.Y)),
-				HorizontalBrush,
+				AllottedGeometry.ToPaintGeometry(FVector2D(0, AllottedGeometry.GetLocalSize().Y - HorizontalImage->ImageSize.Y), FVector2D(AllottedGeometry.Size.X, HorizontalImage->ImageSize.Y)),
+				HorizontalImage,
 				ESlateDrawEffect::None,
 				DashColor);
 
@@ -180,8 +177,8 @@ int32 SDropTarget::OnPaint(const FPaintArgs& Args, const FGeometry& AllottedGeom
 			FSlateDrawElement::MakeBox(
 				OutDrawElements,
 				DashLayer,
-				AllottedGeometry.ToPaintGeometry(FVector2D(0, 0), FVector2D(VerticalBrush->ImageSize.X, AllottedGeometry.GetLocalSize().Y)),
-				VerticalBrush,
+				AllottedGeometry.ToPaintGeometry(FVector2D(0, 0), FVector2D(VerticalImage->ImageSize.X, AllottedGeometry.GetLocalSize().Y)),
+				VerticalImage,
 				ESlateDrawEffect::None,
 				DashColor);
 
@@ -189,8 +186,8 @@ int32 SDropTarget::OnPaint(const FPaintArgs& Args, const FGeometry& AllottedGeom
 			FSlateDrawElement::MakeBox(
 				OutDrawElements,
 				DashLayer,
-				AllottedGeometry.ToPaintGeometry(FVector2D(AllottedGeometry.GetLocalSize().X - VerticalBrush->ImageSize.X, 0), FVector2D(VerticalBrush->ImageSize.X, AllottedGeometry.GetLocalSize().Y)),
-				VerticalBrush,
+				AllottedGeometry.ToPaintGeometry(FVector2D(AllottedGeometry.GetLocalSize().X - VerticalImage->ImageSize.X, 0), FVector2D(VerticalImage->ImageSize.X, AllottedGeometry.GetLocalSize().Y)),
+				VerticalImage,
 				ESlateDrawEffect::None,
 				DashColor);
 
