PR #5126: Toggle visibility of FText properties if unnecessary (Contributed by [REDACTED_PERSON])

diff --git a/Engine/Source/Editor/EditorWidgets/Private/STextPropertyEditableTextBox.cpp b/Engine/Source/Editor/EditorWidgets/Private/STextPropertyEditableTextBox.cpp
--- a/Engine/Source/Editor/EditorWidgets/Private/STextPropertyEditableTextBox.cpp
+++ b/Engine/Source/Editor/EditorWidgets/Private/STextPropertyEditableTextBox.cpp
@@ -584,12 +584,14 @@ void STextPropertyEditableTextBox::Construct(const FArguments& InArgs, const TSh
 					[
 						SNew(STextBlock)
 						.Text(LOCTEXT("TextPackageLabel", "Package:"))
+						.Visibility(this, &STextPropertyEditableTextBox::GetLocalizableVisibility)
 					]
 					+SGridPanel::Slot(1, 2)
 					.Padding(2)
 					[
 						SNew(SEditableTextBox)
 						.Text(this, &STextPropertyEditableTextBox::GetPackageValue)
+						.Visibility(this, &STextPropertyEditableTextBox::GetLocalizableVisibility)
 						.IsReadOnly(true)
 					]
 #endif // USE_STABLE_LOCALIZATION_KEYS
@@ -602,6 +604,7 @@ void STextPropertyEditableTextBox::Construct(const FArguments& InArgs, const TSh
 					[
 						SNew(STextBlock)
 						.Text(LOCTEXT("TextNamespaceLabel", "Namespace:"))
+						.Visibility(this, &STextPropertyEditableTextBox::GetLocalizableVisibility)
 					]
 					+SGridPanel::Slot(1, 3)
 					.Padding(2)
@@ -613,6 +616,7 @@ void STextPropertyEditableTextBox::Construct(const FArguments& InArgs, const TSh
 						.OnTextChanged(this, &STextPropertyEditableTextBox::OnNamespaceChanged)
 						.OnTextCommitted(this, &STextPropertyEditableTextBox::OnNamespaceCommitted)
 						.SelectAllTextOnCommit(true)
+						.Visibility(this, &STextPropertyEditableTextBox::GetLocalizableVisibility)
 						.IsReadOnly(this, &STextPropertyEditableTextBox::IsIdentityReadOnly)
 					]
 
@@ -624,12 +628,14 @@ void STextPropertyEditableTextBox::Construct(const FArguments& InArgs, const TSh
 					[
 						SNew(STextBlock)
 						.Text(LOCTEXT("TextKeyLabel", "Key:"))
+						.Visibility(this, &STextPropertyEditableTextBox::GetLocalizableVisibility)
 					]
 					+SGridPanel::Slot(1, 4)
 					.Padding(2)
 					[
 						SAssignNew(KeyEditableTextBox, SEditableTextBox)
 						.Text(this, &STextPropertyEditableTextBox::GetKeyValue)
+						.Visibility(this, &STextPropertyEditableTextBox::GetLocalizableVisibility)
 #if USE_STABLE_LOCALIZATION_KEYS
 						.SelectAllTextWhenFocused(true)
 						.ClearKeyboardFocusOnCommit(false)
@@ -652,6 +658,7 @@ void STextPropertyEditableTextBox::Construct(const FArguments& InArgs, const TSh
 						SNew(STextBlock)
 						.TextStyle(FEditorStyle::Get(), "LargeText")
 						.Text(LOCTEXT("TextReferencedTextLabel", "Referenced Text"))
+						.Visibility(this, &STextPropertyEditableTextBox::GetLocalizableVisibility)
 					]
 
 					// String Table
@@ -662,6 +669,7 @@ void STextPropertyEditableTextBox::Construct(const FArguments& InArgs, const TSh
 					[
 						SNew(STextBlock)
 						.Text(LOCTEXT("TextStringTableLabel", "String Table:"))
+						.Visibility(this, &STextPropertyEditableTextBox::GetLocalizableVisibility)
 					]
 					+SGridPanel::Slot(1, 6)
 					.Padding(2)
@@ -669,6 +677,7 @@ void STextPropertyEditableTextBox::Construct(const FArguments& InArgs, const TSh
 						SNew(STextPropertyEditableStringTableReference, InEditableTextProperty)
 						.AllowUnlink(true)
 						.IsEnabled(this, &STextPropertyEditableTextBox::CanEdit)
+						.Visibility(this, &STextPropertyEditableTextBox::GetLocalizableVisibility)
 					]
 				]
 			]
@@ -688,6 +697,17 @@ void STextPropertyEditableTextBox::Construct(const FArguments& InArgs, const TSh
 	SetEnabled(TAttribute<bool>(this, &STextPropertyEditableTextBox::CanEdit));
 }
 
+EVisibility STextPropertyEditableTextBox::GetLocalizableVisibility() const
+{
+	const int32 NumTexts = EditableTextProperty->GetNumTexts();
+	if (NumTexts == 1)
+	{
+		const FText PropertyValue = EditableTextProperty->GetText(0);
+		return PropertyValue.IsCultureInvariant() ? EVisibility::Collapsed : EVisibility::Visible;
+	}
+	return EVisibility::Visible;
+}
+
 void STextPropertyEditableTextBox::GetDesiredWidth(float& OutMinDesiredWidth, float& OutMaxDesiredWidth)
 {
 	if (bIsMultiLine)
diff --git a/Engine/Source/Editor/EditorWidgets/Public/STextPropertyEditableTextBox.h b/Engine/Source/Editor/EditorWidgets/Public/STextPropertyEditableTextBox.h
--- a/Engine/Source/Editor/EditorWidgets/Public/STextPropertyEditableTextBox.h
+++ b/Engine/Source/Editor/EditorWidgets/Public/STextPropertyEditableTextBox.h
@@ -173,6 +173,7 @@ private:
 	bool IsSourceTextReadOnly() const;
 	bool IsIdentityReadOnly() const;
 	FText GetToolTipText() const;
+	EVisibility GetLocalizableVisibility() const;
 
 	FText GetTextValue() const;
 	void OnTextChanged(const FText& NewText);
