PR #5391: Make string table dropdown searchable when selecting localization string (Contributed by [REDACTED_PERSON])

diff --git a/Engine/Source/Editor/EditorWidgets/Private/STextPropertyEditableTextBox.cpp b/Engine/Source/Editor/EditorWidgets/Private/STextPropertyEditableTextBox.cpp
--- a/Engine/Source/Editor/EditorWidgets/Private/STextPropertyEditableTextBox.cpp
+++ b/Engine/Source/Editor/EditorWidgets/Private/STextPropertyEditableTextBox.cpp
@@ -14,6 +14,7 @@
 #include "Widgets/Input/SComboButton.h"
 #include "Widgets/Input/SComboBox.h"
 #include "Widgets/Input/SCheckBox.h"
+#include "Widgets/Input/SSearchBox.h"
 #include "EditorStyleSet.h"
 #include "Misc/PackageName.h"
 #include "AssetRegistryModule.h"
@@ -85,17 +86,23 @@ void STextPropertyEditableStringTableReference::Construct(const FArguments& InAr
 {
 	EditableTextProperty = InEditableTextProperty;
 
+	OptionTextFilter = MakeShareable(new FOptionTextFilter(FOptionTextFilter::FItemToStringArray::CreateLambda([](const TSharedPtr<FAvailableStringTable>& InItem, OUT TArray< FString >& StringArray) {
+		StringArray.Add(InItem->DisplayName.ToString());
+	})));
+	KeyTextFilter = MakeShareable(new FKeyTextFilter(FKeyTextFilter::FItemToStringArray::CreateLambda([](const TSharedPtr<FString>& InItem, OUT TArray< FString >& StringArray) {
+		StringArray.Add(*InItem);
+	})));
+
 	TSharedRef<SHorizontalBox> HorizontalBox = SNew(SHorizontalBox);
 
 	HorizontalBox->AddSlot()
 		[
-			SAssignNew(StringTableCombo, SComboBox<TSharedPtr<FAvailableStringTable>>)
-			.ComboBoxStyle(InArgs._ComboStyle)
-			.OptionsSource(&StringTableComboOptions)
-			.OnGenerateWidget(this, &STextPropertyEditableStringTableReference::MakeStringTableComboWidget)
-			.OnSelectionChanged(this, &STextPropertyEditableStringTableReference::OnStringTableComboChanged)
-			.OnComboBoxOpening(this, &STextPropertyEditableStringTableReference::OnStringTableComboOpening)
-			.Content()
+			SAssignNew(StringTableOptionsCombo, SComboButton)
+			.ComboButtonStyle(&InArgs._ComboStyle->ComboButtonStyle)
+			.ContentPadding(FMargin(4.0, 2.0))
+			.OnGetMenuContent(this, &STextPropertyEditableStringTableReference::OnGetStringTableComboOptions)
+			.OnComboBoxOpened(this, &STextPropertyEditableStringTableReference::UpdateStringTableComboOptions)
+			.ButtonContent()
 			[
 				SNew(STextBlock)
 				.Text(this, &STextPropertyEditableStringTableReference::GetStringTableComboContent)
@@ -105,13 +112,13 @@ void STextPropertyEditableStringTableReference::Construct(const FArguments& InAr
 
 	HorizontalBox->AddSlot()
 		[
-			SAssignNew(KeyCombo, SComboBox<TSharedPtr<FString>>)
-			.ComboBoxStyle(InArgs._ComboStyle)
-			.OptionsSource(&KeyComboOptions)
-			.OnGenerateWidget(this, &STextPropertyEditableStringTableReference::MakeKeyComboWidget)
-			.OnSelectionChanged(this, &STextPropertyEditableStringTableReference::OnKeyComboChanged)
-			.OnComboBoxOpening(this, &STextPropertyEditableStringTableReference::OnKeyComboOpening)
-			.Content()
+			SAssignNew(StringTableKeysCombo, SComboButton)
+			.ComboButtonStyle(&InArgs._ComboStyle->ComboButtonStyle)
+			.ContentPadding(FMargin(4.0, 2.0))
+			.IsEnabled(this, &STextPropertyEditableStringTableReference::IsUnlinkEnabled)
+			.OnGetMenuContent(this, &STextPropertyEditableStringTableReference::OnGetStringTableKeyOptions)
+			.OnComboBoxOpened(this, &STextPropertyEditableStringTableReference::UpdateStringTableKeyOptions)
+			.ButtonContent()
 			[
 				SNew(STextBlock)
 				.Text(this, &STextPropertyEditableStringTableReference::GetKeyComboContent)
@@ -126,6 +133,8 @@ void STextPropertyEditableStringTableReference::Construct(const FArguments& InAr
 			[
 				SNew(SButton)
 				.ButtonStyle(InArgs._ButtonStyle)
+				.HAlign(HAlign_Center)
+				.VAlign(VAlign_Center)
 				.Text(LOCTEXT("UnlinkStringTable", "Unlink"))
 				.IsEnabled(this, &STextPropertyEditableStringTableReference::IsUnlinkEnabled)
 				.OnClicked(this, &STextPropertyEditableStringTableReference::OnUnlinkClicked)
@@ -138,6 +147,108 @@ void STextPropertyEditableStringTableReference::Construct(const FArguments& InAr
 	];
 }
 
+void STextPropertyEditableStringTableReference::OnOptionsFilterTextChanged(const FText& InNewText)
+{
+	OptionTextFilter->SetRawFilterText(InNewText);
+	OptionsSearchBox->SetError(OptionTextFilter->GetFilterErrorText());
+
+	UpdateStringTableComboOptions();
+}
+
+void STextPropertyEditableStringTableReference::OnKeysFilterTextChanged(const FText& InNewText)
+{
+	KeyTextFilter->SetRawFilterText(InNewText);
+	KeysSearchBox->SetError(KeyTextFilter->GetFilterErrorText());
+
+	UpdateStringTableKeyOptions();
+}
+
+TSharedRef<SWidget> STextPropertyEditableStringTableReference::OnGetStringTableComboOptions()
+{
+	const FComboButtonStyle& ComboButtonStyle = FCoreStyle::Get().GetWidgetStyle< FComboButtonStyle >("ComboButton");
+	return SNew(SBorder)
+		.BorderImage(&ComboButtonStyle.MenuBorderBrush)
+		.Padding(ComboButtonStyle.MenuBorderPadding)
+		[
+			SNew(SBox)
+			.Padding(4)
+			.WidthOverride(280)
+			.MaxDesiredHeight(600)
+			[
+				SNew(SVerticalBox)
+				+SVerticalBox::Slot()
+				.AutoHeight()
+				[
+					SAssignNew(OptionsSearchBox, SSearchBox)
+					.OnTextChanged(this, &STextPropertyEditableStringTableReference::OnOptionsFilterTextChanged)
+				]
+				+SVerticalBox::Slot()
+				.FillHeight(1.f)
+				.Padding(0, 5, 0, 0)
+				[
+					SAssignNew(StringTableOptionsList, SListView<TSharedPtr<FAvailableStringTable>>)
+					.ListItemsSource(&StringTableComboOptions)
+					.SelectionMode(ESelectionMode::Single)
+					.OnGenerateRow(this, &STextPropertyEditableStringTableReference::OnGenerateStringTableComboOption)
+					.OnSelectionChanged(this, &STextPropertyEditableStringTableReference::OnStringTableComboChanged)
+				]
+			]
+		];
+}
+
+TSharedRef<ITableRow> STextPropertyEditableStringTableReference::OnGenerateStringTableComboOption(TSharedPtr<FAvailableStringTable> InItem, const TSharedRef<STableViewBase>& OwnerTable)
+{
+	return SNew(STableRow<TSharedPtr<FString>>, OwnerTable)
+		[
+			SNew(STextBlock)
+			.Text(InItem->DisplayName)
+			.ToolTipText(FText::FromName(InItem->TableId))
+		];
+}
+
+TSharedRef<SWidget> STextPropertyEditableStringTableReference::OnGetStringTableKeyOptions()
+{
+	const FComboButtonStyle& ComboButtonStyle = FCoreStyle::Get().GetWidgetStyle< FComboButtonStyle >("ComboButton");
+	return SNew(SBorder)
+		.BorderImage(&ComboButtonStyle.MenuBorderBrush)
+		.Padding(ComboButtonStyle.MenuBorderPadding)
+		[
+			SNew(SBox)
+			.Padding(4)
+			.WidthOverride(280)
+			.MaxDesiredHeight(600)
+			[
+				SNew(SVerticalBox)
+				+SVerticalBox::Slot()
+				.AutoHeight()
+				[
+					SAssignNew(KeysSearchBox, SSearchBox)
+					.OnTextChanged(this, &STextPropertyEditableStringTableReference::OnKeysFilterTextChanged)
+				]
+				+SVerticalBox::Slot()
+				.FillHeight(1.f)
+				.Padding(0, 5, 0, 0)
+				[
+					SAssignNew(StringTableKeysList, SListView<TSharedPtr<FString>>)
+					.ListItemsSource(&KeyComboOptions)
+					.SelectionMode(ESelectionMode::Single)
+					.OnGenerateRow(this, &STextPropertyEditableStringTableReference::OnGenerateStringTableKeyOption)
+					.OnSelectionChanged(this, &STextPropertyEditableStringTableReference::OnKeyComboChanged)
+				]
+			]
+		];
+}
+
+TSharedRef<ITableRow> STextPropertyEditableStringTableReference::OnGenerateStringTableKeyOption(TSharedPtr<FString> InItem, const TSharedRef<STableViewBase>& OwnerTable)
+{
+	return SNew(STableRow<TSharedPtr<FString>>, OwnerTable)
+		[
+			SNew(STextBlock)
+			.Text(FText::FromString(*InItem))
+			.ToolTipText(FText::FromString(*InItem))
+		];
+}
+
 void STextPropertyEditableStringTableReference::GetTableIdAndKey(FName& OutTableId, FString& OutKey) const
 {
 	const int32 NumTexts = EditableTextProperty->GetNumTexts();
@@ -183,13 +294,6 @@ void STextPropertyEditableStringTableReference::SetTableIdAndKey(const FName InT
 	}
 }
 
-TSharedRef<SWidget> STextPropertyEditableStringTableReference::MakeStringTableComboWidget(TSharedPtr<FAvailableStringTable> InItem)
-{
-	return SNew(STextBlock)
-		.Text(InItem->DisplayName)
-		.ToolTipText(FText::FromName(InItem->TableId));
-}
-
 void STextPropertyEditableStringTableReference::OnStringTableComboChanged(TSharedPtr<FAvailableStringTable> NewSelection, ESelectInfo::Type SelectInfo)
 {
 	// If it's set from code, we did that on purpose
@@ -208,11 +312,15 @@ void STextPropertyEditableStringTableReference::OnStringTableComboChanged(TShare
 				SetTableIdAndKey(TableId, InKey);
 				return false; // stop enumeration
 			});
+
+			StringTableOptionsCombo->SetIsOpen(false);
+
+			OptionsSearchBox->SetText(FText::GetEmpty());
 		}
 	}
 }
 
-void STextPropertyEditableStringTableReference::OnStringTableComboOpening()
+void STextPropertyEditableStringTableReference::UpdateStringTableComboOptions()
 {
 	FName CurrentTableId;
 	{
@@ -226,12 +334,29 @@ void STextPropertyEditableStringTableReference::OnStringTableComboOpening()
 	// Process assets first (as they may currently be unloaded)
 	{
 		FAssetRegistryModule& AssetRegistryModule = FModuleManager::Get().LoadModuleChecked<FAssetRegistryModule>(AssetRegistryConstants::ModuleName);
-		
+
 		TArray<FAssetData> StringTableAssets;
 		AssetRegistryModule.Get().GetAssetsByClass(UStringTable::StaticClass()->GetFName(), StringTableAssets);
-		
+
 		for (const FAssetData& StringTableAsset : StringTableAssets)
 		{
+			// Only allow string tables assets that have entries to be visible otherwise unexpected behaviour happens for the user
+			bool HasEntries = false;
+			FStringTableConstPtr StringTable = FStringTableRegistry::Get().FindStringTable(StringTableAsset.ObjectPath);
+			if (StringTable.IsValid())
+			{
+				StringTable->EnumerateSourceStrings([&](const FString& InKey, const FString& InSourceString) -> bool
+				{
+					HasEntries = true;
+					return false; // stop enumeration
+				});
+			}
+
+			if (!HasEntries)
+			{
+				continue; // continue on to the next string table asset
+			}
+
 			TSharedRef<FAvailableStringTable> AvailableStringTableEntry = MakeShared<FAvailableStringTable>();
 			AvailableStringTableEntry->TableId = StringTableAsset.ObjectPath;
 			AvailableStringTableEntry->DisplayName = FText::FromName(StringTableAsset.AssetName);
@@ -239,7 +364,10 @@ void STextPropertyEditableStringTableReference::OnStringTableComboOpening()
 			{
 				SelectedStringTableComboEntry = AvailableStringTableEntry;
 			}
-			StringTableComboOptions.Add(AvailableStringTableEntry);
+			if (OptionTextFilter->PassesFilter(AvailableStringTableEntry))
+			{
+				StringTableComboOptions.Add(AvailableStringTableEntry);
+			}
 		}
 	}
 
@@ -251,7 +379,14 @@ void STextPropertyEditableStringTableReference::OnStringTableComboOpening()
 			return InAvailableStringTable->TableId == InTableId;
 		});
 
-		if (!bAlreadyAdded)
+		bool bHasEntries = false;
+		InStringTable->EnumerateSourceStrings([&bHasEntries](const FString& InKey, const FString& InSourceString) -> bool
+		{
+			bHasEntries = true;
+			return false; // stop enumeration
+		});
+
+		if (!bAlreadyAdded && bHasEntries)
 		{
 			TSharedRef<FAvailableStringTable> AvailableStringTableEntry = MakeShared<FAvailableStringTable>();
 			AvailableStringTableEntry->TableId = InTableId;
@@ -260,7 +395,10 @@ void STextPropertyEditableStringTableReference::OnStringTableComboOpening()
 			{
 				SelectedStringTableComboEntry = AvailableStringTableEntry;
 			}
-			StringTableComboOptions.Add(AvailableStringTableEntry);
+			if (OptionTextFilter->PassesFilter(AvailableStringTableEntry))
+			{
+				StringTableComboOptions.Add(AvailableStringTableEntry);
+			}
 		}
 
 		return true; // continue enumeration
@@ -271,13 +409,15 @@ void STextPropertyEditableStringTableReference::OnStringTableComboOpening()
 		return InOne->DisplayName.ToString() < InTwo->DisplayName.ToString();
 	});
 
+	StringTableOptionsList->RebuildList();
+
 	if (SelectedStringTableComboEntry.IsValid())
 	{
-		StringTableCombo->SetSelectedItem(SelectedStringTableComboEntry);
+		StringTableOptionsList->SetItemSelection(SelectedStringTableComboEntry, true);
 	}
 	else
 	{
-		StringTableCombo->ClearSelection();
+		StringTableOptionsList->ClearSelection();
 	}
 }
 
@@ -303,13 +443,6 @@ FText STextPropertyEditableStringTableReference::GetStringTableComboToolTip() co
 	return FText::FromName(CurrentTableId);
 }
 
-TSharedRef<SWidget> STextPropertyEditableStringTableReference::MakeKeyComboWidget(TSharedPtr<FString> InItem)
-{
-	return SNew(STextBlock)
-		.Text(FText::FromString(*InItem))
-		.ToolTipText(FText::FromString(*InItem));
-}
-
 void STextPropertyEditableStringTableReference::OnKeyComboChanged(TSharedPtr<FString> NewSelection, ESelectInfo::Type SelectInfo)
 {
 	// If it's set from code, we did that on purpose
@@ -322,10 +455,14 @@ void STextPropertyEditableStringTableReference::OnKeyComboChanged(TSharedPtr<FSt
 		}
 
 		SetTableIdAndKey(CurrentTableId, *NewSelection);
+
+		StringTableKeysCombo->SetIsOpen(false);
+
+		KeysSearchBox->SetText(FText::GetEmpty());
 	}
 }
 
-void STextPropertyEditableStringTableReference::OnKeyComboOpening()
+void STextPropertyEditableStringTableReference::UpdateStringTableKeyOptions()
 {
 	FName CurrentTableId;
 	FString CurrentKey;
@@ -346,7 +483,10 @@ void STextPropertyEditableStringTableReference::OnKeyComboOpening()
 				{
 					SelectedKeyComboEntry = KeyComboEntry;
 				}
-				KeyComboOptions.Add(KeyComboEntry);
+				if (KeyTextFilter->PassesFilter(KeyComboEntry))
+				{
+					KeyComboOptions.Add(KeyComboEntry);
+				}
 				return true; // continue enumeration
 			});
 		}
@@ -357,13 +497,15 @@ void STextPropertyEditableStringTableReference::OnKeyComboOpening()
 		return *InOne < *InTwo;
 	});
 
+	StringTableKeysList->RebuildList();
+
 	if (SelectedKeyComboEntry.IsValid())
 	{
-		KeyCombo->SetSelectedItem(SelectedKeyComboEntry);
+		StringTableKeysList->SetItemSelection(SelectedKeyComboEntry, true);
 	}
 	else
 	{
-		KeyCombo->ClearSelection();
+		StringTableKeysList->ClearSelection();
 	}
 }
 
diff --git a/Engine/Source/Editor/EditorWidgets/Public/STextPropertyEditableTextBox.h b/Engine/Source/Editor/EditorWidgets/Public/STextPropertyEditableTextBox.h
--- a/Engine/Source/Editor/EditorWidgets/Public/STextPropertyEditableTextBox.h
+++ b/Engine/Source/Editor/EditorWidgets/Public/STextPropertyEditableTextBox.h
@@ -14,12 +14,14 @@
 #include "Styling/SlateTypes.h"
 #include "Styling/SlateWidgetStyleAsset.h"
 #include "Types/SlateStructs.h"
+#include "Widgets/Views/SListView.h"
+#include "Misc/TextFilter.h"
 
 class SEditableTextBox;
 class SMultiLineEditableTextBox;
 
-template <typename OptionType>
-class SComboBox;
+class SComboButton;
+class SSearchBox;
 
 /** Interface to allow STextPropertyEditableTextBox to be used to edit both properties and Blueprint pins */
 class IEditableTextProperty
@@ -32,7 +34,7 @@ public:
 		EditedSource,
 	};
 
-	virtual ~IEditableTextProperty() {}
+	virtual ~IEditableTextProperty() = default;
 
 	/** Are the text properties being edited marked as multi-line? */
 	virtual bool IsMultiLineText() const = 0;
@@ -108,15 +110,21 @@ private:
 	void GetTableIdAndKey(FName& OutTableId, FString& OutKey) const;
 	void SetTableIdAndKey(const FName InTableId, const FString& InKey);
 
-	TSharedRef<SWidget> MakeStringTableComboWidget(TSharedPtr<FAvailableStringTable> InItem);
+	void OnOptionsFilterTextChanged(const FText& InNewText);
+	void OnKeysFilterTextChanged(const FText& InNewText);
+
+	TSharedRef<SWidget> OnGetStringTableComboOptions();
+	TSharedRef<class ITableRow> OnGenerateStringTableComboOption(TSharedPtr<FAvailableStringTable> InItem, const TSharedRef<class STableViewBase>& OwnerTable);
+	TSharedRef<SWidget> OnGetStringTableKeyOptions();
+	TSharedRef<class ITableRow> OnGenerateStringTableKeyOption(TSharedPtr<FString> InItem, const TSharedRef<class STableViewBase>& OwnerTable);
+
 	void OnStringTableComboChanged(TSharedPtr<FAvailableStringTable> NewSelection, ESelectInfo::Type SelectInfo);
-	void OnStringTableComboOpening();
+	void UpdateStringTableComboOptions();
 	FText GetStringTableComboContent() const;
 	FText GetStringTableComboToolTip() const;
 
-	TSharedRef<SWidget> MakeKeyComboWidget(TSharedPtr<FString> InItem);
 	void OnKeyComboChanged(TSharedPtr<FString> NewSelection, ESelectInfo::Type SelectInfo);
-	void OnKeyComboOpening();
+	void UpdateStringTableKeyOptions();
 	FText GetKeyComboContent() const;
 	FText GetKeyComboToolTip() const;
 
@@ -125,10 +133,20 @@ private:
 
 	TSharedPtr<IEditableTextProperty> EditableTextProperty;
 
-	TSharedPtr<SComboBox<TSharedPtr<FAvailableStringTable>>> StringTableCombo;
-	TArray<TSharedPtr<FAvailableStringTable>> StringTableComboOptions;
+	using FOptionTextFilter = TTextFilter< TSharedPtr<FAvailableStringTable> >;
+	TSharedPtr<FOptionTextFilter> OptionTextFilter;
+	TSharedPtr<SSearchBox> OptionsSearchBox;
 
-	TSharedPtr<SComboBox<TSharedPtr<FString>>> KeyCombo;
+	using FKeyTextFilter = TTextFilter< TSharedPtr<FString> >;
+	TSharedPtr<FKeyTextFilter> KeyTextFilter;
+	TSharedPtr<SSearchBox> KeysSearchBox;
+
+	TSharedPtr<SComboButton> StringTableOptionsCombo;
+	TSharedPtr<SListView<TSharedPtr<FAvailableStringTable>>> StringTableOptionsList;
+	TSharedPtr<SComboButton> StringTableKeysCombo;
+	TSharedPtr<SListView<TSharedPtr<FString>>> StringTableKeysList;
+
+	TArray<TSharedPtr<FAvailableStringTable>> StringTableComboOptions;
 	TArray<TSharedPtr<FString>> KeyComboOptions;
 };
 
