Make SAssetDiscoveryIndicator SelfHitTestInvisible so it stops eating mouse events

diff --git a/Engine/Source/Editor/EditorWidgets/Private/SAssetDiscoveryIndicator.h b/Engine/Source/Editor/EditorWidgets/Private/SAssetDiscoveryIndicator.h
--- a/Engine/Source/Editor/EditorWidgets/Private/SAssetDiscoveryIndicator.h
+++ b/Engine/Source/Editor/EditorWidgets/Private/SAssetDiscoveryIndicator.h
@@ -20,7 +20,11 @@ public:
 	SLATE_BEGIN_ARGS( SAssetDiscoveryIndicator )
 		: _ScaleMode(EAssetDiscoveryIndicatorScaleMode::Scale_None)
 		, _FadeIn(true)
-		{}
+		{
+			// This widget has no interactive controls and is used as an overlay,
+			// so removing it from the hit test grid will stop it from eating scroll events.
+			_Visibility = EVisibility::SelfHitTestInvisible;
+		}
 
 		/** The way the indicator will scale out when done displaying progress */
 		SLATE_ARGUMENT( EAssetDiscoveryIndicatorScaleMode::Type, ScaleMode )
