Move duplicated SEnumComboBox impl in Niagara out to Engine EditorWidgets.
Add Disabled State handling and custom tooltip delegate for SEnumComboBox.

diff --git a/Engine/Plugins/FX/Niagara/Source/NiagaraEditor/Private/TypeEditorUtilities/NiagaraEnumTypeEditorUtilities.cpp b/Engine/Plugins/FX/Niagara/Source/NiagaraEditor/Private/TypeEditorUtilities/NiagaraEnumTypeEditorUtilities.cpp
--- a/Engine/Plugins/FX/Niagara/Source/NiagaraEditor/Private/TypeEditorUtilities/NiagaraEnumTypeEditorUtilities.cpp
+++ b/Engine/Plugins/FX/Niagara/Source/NiagaraEditor/Private/TypeEditorUtilities/NiagaraEnumTypeEditorUtilities.cpp
@@ -5,9 +5,9 @@
 #include "NiagaraTypes.h"
 #include "NiagaraEditorStyle.h"
 #include "NiagaraEditorCommon.h"
-#include "SEnumCombobox.h"
 
 #include "Widgets/Input/SSpinBox.h"
+#include "EditorWidgets/Public/SEnumComboBox.h"
 
 class SNiagaraEnumParameterEditor : public SNiagaraParameterEditor
 {
diff --git a/Engine/Source/Editor/EditorWidgets/Private/SEnumComboBox.cpp b/Engine/Source/Editor/EditorWidgets/Private/SEnumComboBox.cpp
new file mode 100644
--- /dev/null
+++ b/Engine/Source/Editor/EditorWidgets/Private/SEnumComboBox.cpp
@@ -0,0 +1,93 @@
+// Copyright Epic Games, Inc. All Rights Reserved.
+
+#include "SEnumComboBox.h"
+#include "EditorStyleSet.h"
+#include "Types/SlateEnums.h"
+#include "Widgets/SToolTip.h"
+
+void SEnumComboBox::Construct(const FArguments& InArgs, const UEnum* InEnum)
+{
+	Enum = InEnum;
+	CurrentValue = InArgs._CurrentValue;
+	check(CurrentValue.IsBound());
+
+	OnGetToolTipForValue = InArgs._OnGetToolTipForValue;
+	Font = InArgs._Font;
+
+	bUpdatingSelectionInternally = false;
+
+	for (int32 i = 0; i < Enum->NumEnums() - 1; i++)
+	{
+		if (Enum->HasMetaData(TEXT("Hidden"), i) == false)
+		{
+			VisibleEnumNameIndices.Add(MakeShareable(new int32(i)));
+		}
+	}
+
+	SComboBox::Construct(SComboBox<TSharedPtr<int32>>::FArguments()
+		.ButtonStyle(InArgs._ButtonStyle)
+		.OptionsSource(&VisibleEnumNameIndices)
+		.OnGenerateWidget(this, &SEnumComboBox::OnGenerateWidget)
+		.OnSelectionChanged(this, &SEnumComboBox::OnComboSelectionChanged)
+		.OnComboBoxOpening(this, &SEnumComboBox::OnComboMenuOpening)
+		.ContentPadding(InArgs._ContentPadding)
+		[
+			SNew(STextBlock)
+			.Font(Font)
+			.Text(this, &SEnumComboBox::GetCurrentValueText)
+			.ToolTipText(this, &SEnumComboBox::GetCurrentValueTooltip)
+		]);
+}
+
+
+FText SEnumComboBox::GetCurrentValueText() const
+{
+	int32 ValueNameIndex = Enum->GetIndexByValue(CurrentValue.Get());
+	return Enum->GetDisplayNameTextByIndex(ValueNameIndex);
+}
+
+FText SEnumComboBox::GetCurrentValueTooltip() const
+{
+	int32 ValueNameIndex = Enum->GetIndexByValue(CurrentValue.Get());
+	if (OnGetToolTipForValue.IsBound())
+	{
+		return OnGetToolTipForValue.Execute(ValueNameIndex);
+	}
+	return Enum->GetToolTipTextByIndex(ValueNameIndex);
+}
+
+TSharedRef<SWidget> SEnumComboBox::OnGenerateWidget(TSharedPtr<int32> InItem)
+{
+	return SNew(STextBlock)
+		.Font(Font)
+		.Text(this, &SEnumComboBox::GetCurrentValueText)
+		.ToolTipText(this, &SEnumComboBox::GetCurrentValueTooltip);
+}
+
+void SEnumComboBox::OnComboSelectionChanged(TSharedPtr<int32> InSelectedItem, ESelectInfo::Type SelectInfo)
+{
+	if (bUpdatingSelectionInternally == false)
+	{
+		OnEnumSelectionChanged.ExecuteIfBound(*InSelectedItem, SelectInfo);
+	}
+}
+
+void SEnumComboBox::OnComboMenuOpening()
+{
+	int32 CurrentNameIndex = Enum->GetIndexByValue(CurrentValue.Get());
+	TSharedPtr<int32> FoundNameIndexItem;
+	for (int32 i = 0; i < VisibleEnumNameIndices.Num(); i++)
+	{
+		if (*VisibleEnumNameIndices[i] == CurrentNameIndex)
+		{
+			FoundNameIndexItem = VisibleEnumNameIndices[i];
+			break;
+		}
+	}
+	if (FoundNameIndexItem.IsValid())
+	{
+		bUpdatingSelectionInternally = true;
+		SetSelectedItem(FoundNameIndexItem);
+		bUpdatingSelectionInternally = false;
+	}
+}
diff --git a/Engine/Source/Editor/EditorWidgets/Public/SEnumCombobox.h b/Engine/Source/Editor/EditorWidgets/Public/SEnumCombobox.h
--- a/Engine/Source/Editor/EditorWidgets/Public/SEnumCombobox.h
+++ b/Engine/Source/Editor/EditorWidgets/Public/SEnumCombobox.h
@@ -3,130 +3,51 @@
 #pragma once
 
 #include "CoreMinimal.h"
-#include "EditorStyleSet.h"
-#include "Types/SlateEnums.h"
 #include "Widgets/Input/SComboBox.h"
-#include "Widgets/SToolTip.h"
 
-/*
-* Editor only widget for creating a combo-box which allows the user to pick the values from a dropdown combo-box. 
-*/
-class SEnumComboBox : public SComboBox<TSharedPtr<int32>>
- {
+class EDITORWIDGETS_API SEnumComboBox : public SComboBox<TSharedPtr<int32>>
+{
 public:
-	DECLARE_DELEGATE_TwoParams(FOnEnumSelectionChanged, int32 /*Selection*/, ESelectInfo::Type /*SelectionType*/);
+	DECLARE_DELEGATE_TwoParams(FOnEnumSelectionChanged, int32, ESelectInfo::Type);
+	DECLARE_DELEGATE_RetVal_OneParam(FText, FOnGetToolTipForValue, int32 /* Value */)
+
 public:
 	SLATE_BEGIN_ARGS(SEnumComboBox)
 		: _CurrentValue()
 		, _ContentPadding(FMargin(4.0, 2.0))
-		, _OnEnumSelectionChanged()
 		, _ButtonStyle(nullptr)
 	{}
 
-		SLATE_ATTRIBUTE(int32, CurrentValue)
+	SLATE_ATTRIBUTE(int32, CurrentValue)
 		SLATE_ATTRIBUTE(FMargin, ContentPadding)
 		SLATE_ATTRIBUTE(FSlateFontInfo, Font)
-		SLATE_ARGUMENT(FOnEnumSelectionChanged, OnEnumSelectionChanged)
+		SLATE_EVENT(FOnEnumSelectionChanged, OnEnumSelectionChanged)
+		SLATE_EVENT(FOnGetToolTipForValue, OnGetToolTipForValue)
 		SLATE_STYLE_ARGUMENT(FButtonStyle, ButtonStyle)
-
 	SLATE_END_ARGS()
 
-	void Construct(const FArguments& InArgs, const UEnum* InEnum)
-	{
-		Enum = InEnum;
-		CurrentValue = InArgs._CurrentValue;
-		check(CurrentValue.IsBound());
-		OnEnumSelectionChangedDelegate = InArgs._OnEnumSelectionChanged;
-		Font = InArgs._Font;
-		bUpdatingSelectionInternally = false;
-
-		for (int32 i = 0; i < Enum->NumEnums() - 1; i++)
-		{
-			if (Enum->HasMetaData(TEXT("Hidden"), i) == false)
-			{
-				VisibleEnumNameIndices.Add(MakeShareable(new int32(i)));
-			}
-		}
-		
-		SComboBox::Construct(SComboBox<TSharedPtr<int32>>::FArguments()
-			.ButtonStyle(InArgs._ButtonStyle)
-			.OptionsSource(&VisibleEnumNameIndices)
-			.OnGenerateWidget_Lambda([this](TSharedPtr<int32> InItem)
-			{
-				return SNew(STextBlock)
-					.Font(Font)
-					.Text(Enum->GetDisplayNameTextByIndex(*InItem))
-					.ToolTipText(Enum->GetToolTipTextByIndex(*InItem));
-			})
-			.OnSelectionChanged(this, &SEnumComboBox::OnComboSelectionChanged)
-			.OnComboBoxOpening(this, &SEnumComboBox::OnComboMenuOpening)
-			.ContentPadding(InArgs._ContentPadding)
-			[
-				SNew(STextBlock)
-					.Font(Font)
-					.Text(this, &SEnumComboBox::GetCurrentValue)
-					.ToolTipText(this, &SEnumComboBox::GetCurrentValueTooltip)
-			]);
-	}
+	void Construct(const FArguments& InArgs, const UEnum* InEnum);
 
 private:
-	FText GetCurrentValue() const
-	{
-		int32 CurrentNameIndex = Enum->GetIndexByValue(CurrentValue.Get());
-		return Enum->GetDisplayNameTextByIndex(CurrentNameIndex);
-	}
-
-	FText GetCurrentValueTooltip() const
-	{
-		int32 CurrentNameIndex = Enum->GetIndexByValue(CurrentValue.Get());
-		return Enum->GetToolTipTextByIndex(CurrentNameIndex);
-	}
-
-	TSharedRef<SWidget> OnGenerateWidget(TSharedPtr<int32> InItem)
-	{
-		return SNew(STextBlock)
-			.Text(Enum->GetDisplayNameTextByIndex(*InItem))
-			.Font(Font);
-	}
-
-	void OnComboSelectionChanged(TSharedPtr<int32> InSelectedItem, ESelectInfo::Type SelectInfo)
-	{
-		if (bUpdatingSelectionInternally == false)
-		{
-			OnEnumSelectionChangedDelegate.ExecuteIfBound(*InSelectedItem, SelectInfo);
-		}
-	}
-
-	void OnComboMenuOpening()
-	{
-		int32 CurrentNameIndex = Enum->GetIndexByValue(CurrentValue.Get());
-		TSharedPtr<int32> FoundNameIndexItem;
-		for (int32 i = 0; i < VisibleEnumNameIndices.Num(); i++)
-		{
-			if (*VisibleEnumNameIndices[i] == CurrentNameIndex)
-			{
-				FoundNameIndexItem = VisibleEnumNameIndices[i];
-				break;
-			}
-		}
-		if (FoundNameIndexItem.IsValid())
-		{
-			bUpdatingSelectionInternally = true;
-			SetSelectedItem(FoundNameIndexItem);
-			bUpdatingSelectionInternally = false;
-		}
-	}
+	FText GetCurrentValueText() const;
+	FText GetCurrentValueTooltip() const;
+	bool GetValueIsEnabled(TWeakPtr<int32> InValueWeak) const;
+	TSharedRef<SWidget> OnGenerateWidget(TSharedPtr<int32> InItem);
+	void OnComboSelectionChanged(TSharedPtr<int32> InSelectedItem, ESelectInfo::Type SelectInfo);
+	void OnComboMenuOpening();
 
 private:
 	const UEnum* Enum;
 
 	TAttribute<int32> CurrentValue;
 
-	TAttribute< FSlateFontInfo > Font;
+	TAttribute<FSlateFontInfo> Font;
 
 	TArray<TSharedPtr<int32>> VisibleEnumNameIndices;
 
 	bool bUpdatingSelectionInternally;
 
-	FOnEnumSelectionChanged OnEnumSelectionChangedDelegate;
-};
\ No newline at end of file
+	FOnGetToolTipForValue OnGetToolTipForValue;
+
+	FOnEnumSelectionChanged OnEnumSelectionChanged;
+};
