Fix Linux build break (SEnumComboBox.h -> SEnumCombobox.h)

diff --git a/Engine/Source/Editor/EditorWidgets/Private/SEnumComboBox.cpp b/Engine/Source/Editor/EditorWidgets/Private/SEnumComboBox.cpp
--- a/Engine/Source/Editor/EditorWidgets/Private/SEnumComboBox.cpp
+++ b/Engine/Source/Editor/EditorWidgets/Private/SEnumComboBox.cpp
@@ -1,6 +1,6 @@
 // Copyright Epic Games, Inc. All Rights Reserved.
 
-#include "SEnumComboBox.h"
+#include "SEnumCombobox.h"
 #include "EditorStyleSet.h"
 #include "Types/SlateEnums.h"
 #include "Widgets/SToolTip.h"
