Fix SEnumComboBoxes not being configurable, and failing to display available entries when expanded.

diff --git a/Engine/Source/Editor/EditorWidgets/Private/SEnumComboBox.cpp b/Engine/Source/Editor/EditorWidgets/Private/SEnumComboBox.cpp
--- a/Engine/Source/Editor/EditorWidgets/Private/SEnumComboBox.cpp
+++ b/Engine/Source/Editor/EditorWidgets/Private/SEnumComboBox.cpp
@@ -10,10 +10,9 @@ void SEnumComboBox::Construct(const FArguments& InArgs, const UEnum* InEnum)
 	Enum = InEnum;
 	CurrentValue = InArgs._CurrentValue;
 	check(CurrentValue.IsBound());
-
+	OnEnumSelectionChangedDelegate = InArgs._OnEnumSelectionChanged;
 	OnGetToolTipForValue = InArgs._OnGetToolTipForValue;
 	Font = InArgs._Font;
-
 	bUpdatingSelectionInternally = false;
 
 	for (int32 i = 0; i < Enum->NumEnums() - 1; i++)
@@ -60,15 +59,14 @@ TSharedRef<SWidget> SEnumComboBox::OnGenerateWidget(TSharedPtr<int32> InItem)
 {
 	return SNew(STextBlock)
 		.Font(Font)
-		.Text(this, &SEnumComboBox::GetCurrentValueText)
-		.ToolTipText(this, &SEnumComboBox::GetCurrentValueTooltip);
+		.Text(Enum->GetDisplayNameTextByIndex(*InItem));
 }
 
 void SEnumComboBox::OnComboSelectionChanged(TSharedPtr<int32> InSelectedItem, ESelectInfo::Type SelectInfo)
 {
 	if (bUpdatingSelectionInternally == false)
 	{
-		OnEnumSelectionChanged.ExecuteIfBound(*InSelectedItem, SelectInfo);
+		OnEnumSelectionChangedDelegate.ExecuteIfBound(*InSelectedItem, SelectInfo);
 	}
 }
 
diff --git a/Engine/Source/Editor/EditorWidgets/Public/SEnumCombobox.h b/Engine/Source/Editor/EditorWidgets/Public/SEnumCombobox.h
--- a/Engine/Source/Editor/EditorWidgets/Public/SEnumCombobox.h
+++ b/Engine/Source/Editor/EditorWidgets/Public/SEnumCombobox.h
@@ -18,7 +18,7 @@ public:
 		, _ButtonStyle(nullptr)
 	{}
 
-	SLATE_ATTRIBUTE(int32, CurrentValue)
+		SLATE_ATTRIBUTE(int32, CurrentValue)
 		SLATE_ATTRIBUTE(FMargin, ContentPadding)
 		SLATE_ATTRIBUTE(FSlateFontInfo, Font)
 		SLATE_EVENT(FOnEnumSelectionChanged, OnEnumSelectionChanged)
@@ -49,5 +49,5 @@ private:
 
 	FOnGetToolTipForValue OnGetToolTipForValue;
 
-	FOnEnumSelectionChanged OnEnumSelectionChanged;
+	FOnEnumSelectionChanged OnEnumSelectionChangedDelegate;
 };
