Restoring previous behavior on SEnumCombox

diff --git a/Engine/Source/Editor/EditorWidgets/Private/SEnumComboBox.cpp b/Engine/Source/Editor/EditorWidgets/Private/SEnumComboBox.cpp
--- a/Engine/Source/Editor/EditorWidgets/Private/SEnumComboBox.cpp
+++ b/Engine/Source/Editor/EditorWidgets/Private/SEnumComboBox.cpp
@@ -59,7 +59,8 @@ TSharedRef<SWidget> SEnumComboBox::OnGenerateWidget(TSharedPtr<int32> InItem)
 {
 	return SNew(STextBlock)
 		.Font(Font)
-		.Text(Enum->GetDisplayNameTextByIndex(*InItem));
+		.Text(Enum->GetDisplayNameTextByIndex(*InItem))
+		.ToolTipText(Enum->GetToolTipTextByIndex(*InItem));
 }
 
 void SEnumComboBox::OnComboSelectionChanged(TSharedPtr<int32> InSelectedItem, ESelectInfo::Type SelectInfo)
