#UE4 Route more references through the game asset filter. Drag/Dropping an actor in a level is now checked. Also you'll now get more detailed feedback when you cannot assign a reference in a property when dragdropping

diff --git a/Engine/Source/Editor/EditorWidgets/Private/SAssetDropTarget.cpp b/Engine/Source/Editor/EditorWidgets/Private/SAssetDropTarget.cpp
--- a/Engine/Source/Editor/EditorWidgets/Private/SAssetDropTarget.cpp
+++ b/Engine/Source/Editor/EditorWidgets/Private/SAssetDropTarget.cpp
@@ -4,6 +4,7 @@
 #include "AssetData.h"
 #include "DragAndDrop/AssetDragDropOp.h"
 #include "DragAndDrop/ActorDragDropOp.h"
+#include "DragAndDrop/DecoratedDragDropOp.h"
 #include "AssetSelection.h"
 
 #define LOCTEXT_NAMESPACE "EditorWidgets"
@@ -12,6 +13,7 @@ void SAssetDropTarget::Construct(const FArguments& InArgs )
 {
 	OnAssetDropped = InArgs._OnAssetDropped;
 	OnIsAssetAcceptableForDrop = InArgs._OnIsAssetAcceptableForDrop;
+	OnIsAssetAcceptableForDropWithReason = InArgs._OnIsAssetAcceptableForDropWithReason;
 
 	SDropTarget::Construct(
 		SDropTarget::FArguments()
@@ -42,7 +44,31 @@ bool SAssetDropTarget::OnAllowDrop(TSharedPtr<FDragDropOperation> DragDropOperat
 	if ( Object )
 	{
 		// Check and see if its valid to drop this object
-		if ( OnIsAssetAcceptableForDrop.IsBound() )
+		if (OnIsAssetAcceptableForDropWithReason.IsBound())
+		{
+			FText FailureReason;
+			if (OnIsAssetAcceptableForDropWithReason.Execute(Object, FailureReason))
+			{
+				return true;
+			}
+			else
+			{
+				if (IsDragOver() && !FailureReason.IsEmpty())
+				{
+					if (DragDropOperation.IsValid() && DragDropOperation->IsOfType<FDecoratedDragDropOp>())
+					{
+						TSharedPtr<FDecoratedDragDropOp> DragDropOp = StaticCastSharedPtr<FDecoratedDragDropOp>(DragDropOperation);
+						if (DragDropOp.IsValid())
+						{
+							DragDropOp->SetToolTip(FailureReason, FEditorStyle::GetBrush(TEXT("Graph.ConnectorFeedback.Error")));
+						}
+					}
+				}
+
+				return false;
+			}
+		}
+		else if ( OnIsAssetAcceptableForDrop.IsBound() )
 		{
 			return OnIsAssetAcceptableForDrop.Execute(Object);
 		}
@@ -64,6 +90,18 @@ bool SAssetDropTarget::OnIsRecognized(TSharedPtr<FDragDropOperation> DragDropOpe
 	return bRecognizedEvent;
 }
 
+void SAssetDropTarget::OnDragLeave(const FDragDropEvent& DragDropEvent)
+{
+	SDropTarget::OnDragLeave(DragDropEvent);
+
+	TSharedPtr<FDragDropOperation> Operation = DragDropEvent.GetOperation();
+	if (Operation.IsValid() && Operation->IsOfType<FDecoratedDragDropOp>())
+	{
+		TSharedPtr<FDecoratedDragDropOp> DragDropOp = StaticCastSharedPtr<FDecoratedDragDropOp>(Operation);
+		DragDropOp->ResetToDefaultToolTip();
+	}
+}
+
 UObject* SAssetDropTarget::GetDroppedObject(TSharedPtr<FDragDropOperation> DragDropOperation, bool& bOutRecognizedEvent) const
 {
 	bOutRecognizedEvent = false;
diff --git a/Engine/Source/Editor/EditorWidgets/Public/SAssetDropTarget.h b/Engine/Source/Editor/EditorWidgets/Public/SAssetDropTarget.h
--- a/Engine/Source/Editor/EditorWidgets/Public/SAssetDropTarget.h
+++ b/Engine/Source/Editor/EditorWidgets/Public/SAssetDropTarget.h
@@ -19,6 +19,9 @@ public:
 	/** Called when we need to check if an asset type is valid for dropping */
 	DECLARE_DELEGATE_RetVal_OneParam( bool, FIsAssetAcceptableForDrop, const UObject* );
 
+	/** Called when we need to check if an asset type is valid for dropping and also will have a reason if it is not */
+	DECLARE_DELEGATE_RetVal_TwoParams( bool, FIsAssetAcceptableForDropWithReason, const UObject*, FText& );
+
 	SLATE_BEGIN_ARGS(SAssetDropTarget)
 	{ }
 		/* Content to display for the in the drop target */
@@ -27,6 +30,8 @@ public:
 		SLATE_EVENT( FOnAssetDropped, OnAssetDropped )
 		/** Called to check if an asset is acceptible for dropping */
 		SLATE_EVENT( FIsAssetAcceptableForDrop, OnIsAssetAcceptableForDrop )
+		/** Called to check if an asset is acceptible for dropping if you also plan on returning a reason text */
+		SLATE_EVENT( FIsAssetAcceptableForDropWithReason, OnIsAssetAcceptableForDropWithReason )
 	SLATE_END_ARGS()
 
 	void Construct(const FArguments& InArgs );
@@ -35,6 +40,7 @@ protected:
 	FReply OnDropped(TSharedPtr<FDragDropOperation> DragDropOperation);
 	virtual bool OnAllowDrop(TSharedPtr<FDragDropOperation> DragDropOperation) const override;
 	virtual bool OnIsRecognized(TSharedPtr<FDragDropOperation> DragDropOperation) const override;
+	virtual void OnDragLeave(const FDragDropEvent& DragDropEvent) override;
 
 private:
 	UObject* GetDroppedObject(TSharedPtr<FDragDropOperation> DragDropOperation, bool& bOutRecognizedEvent) const;
@@ -44,4 +50,6 @@ private:
 	FOnAssetDropped OnAssetDropped;
 	/** Delegate to call to check validity of the asset */
 	FIsAssetAcceptableForDrop OnIsAssetAcceptableForDrop;
+	/** Delegate to call to check validity of the asset if you will also provide a reason when returning false */
+	FIsAssetAcceptableForDropWithReason OnIsAssetAcceptableForDropWithReason;
 };
diff --git a/Engine/Source/Editor/EditorWidgets/Public/SDropTarget.h b/Engine/Source/Editor/EditorWidgets/Public/SDropTarget.h
--- a/Engine/Source/Editor/EditorWidgets/Public/SDropTarget.h
+++ b/Engine/Source/Editor/EditorWidgets/Public/SDropTarget.h
@@ -81,6 +81,9 @@ protected:
 	/** Get the brightness on the background. */
 	FSlateColor GetBackgroundBrightness() const;
 
+	/** Returns true if this drop target is currently being hovered over by a drag drop event */
+	bool IsDragOver() const { return bIsDragOver; }
+
 private:
 	/** Delegate to call when an asset is dropped */
 	FOnDrop DroppedEvent;
diff --git a/Engine/Source/Editor/LevelEditor/Private/SLevelViewport.cpp b/Engine/Source/Editor/LevelEditor/Private/SLevelViewport.cpp
--- a/Engine/Source/Editor/LevelEditor/Private/SLevelViewport.cpp
+++ b/Engine/Source/Editor/LevelEditor/Private/SLevelViewport.cpp
@@ -660,9 +660,16 @@ void SLevelViewport::OnDragLeave( const FDragDropEvent& DragDropEvent )
 		LevelViewportClient->DestroyDropPreviewActors();
 	}
 
-	if (DragDropEvent.GetOperation().IsValid())
+	TSharedPtr<FDragDropOperation> Operation = DragDropEvent.GetOperation();
+	if (Operation.IsValid())
 	{
-		DragDropEvent.GetOperation()->SetDecoratorVisibility(true);
+		Operation->SetDecoratorVisibility(true);
+
+		if (Operation->IsOfType<FDecoratedDragDropOp>())
+		{
+			TSharedPtr<FDecoratedDragDropOp> DragDropOp = StaticCastSharedPtr<FDecoratedDragDropOp>(Operation);
+			DragDropOp->ResetToDefaultToolTip();
+		}
 	}
 }
 
@@ -752,7 +759,9 @@ bool SLevelViewport::HandleDragObjects(const FGeometry& MyGeometry, const FDragD
 		{
 			// At least one of the assets can't be dropped.
 			Operation->SetCursorOverride(EMouseCursor::SlashedCircle);
-			return false;
+			bValidDrag = false;
+			HintText = DropResult.HintText;
+			break;
 		}
 		else
 		{
diff --git a/Engine/Source/Editor/PropertyEditor/Private/UserInterface/PropertyEditor/SPropertyEditorAsset.cpp b/Engine/Source/Editor/PropertyEditor/Private/UserInterface/PropertyEditor/SPropertyEditorAsset.cpp
--- a/Engine/Source/Editor/PropertyEditor/Private/UserInterface/PropertyEditor/SPropertyEditorAsset.cpp
+++ b/Engine/Source/Editor/PropertyEditor/Private/UserInterface/PropertyEditor/SPropertyEditorAsset.cpp
@@ -321,7 +321,7 @@ void SPropertyEditorAsset::Construct(const FArguments& InArgs, const TSharedPtr<
 	ChildSlot
 	[
 		SNew( SAssetDropTarget )
-		.OnIsAssetAcceptableForDrop( this, &SPropertyEditorAsset::OnAssetDraggedOver )
+		.OnIsAssetAcceptableForDropWithReason( this, &SPropertyEditorAsset::OnAssetDraggedOver )
 		.OnAssetDropped( this, &SPropertyEditorAsset::OnAssetDropped )
 		[
 			SAssignNew( ValueContentBox, SHorizontalBox )	
@@ -1208,7 +1208,7 @@ FSlateColor SPropertyEditorAsset::GetAssetClassColor()
 	return FSlateColor::UseForeground();
 }
 
-bool SPropertyEditorAsset::OnAssetDraggedOver( const UObject* InObject ) const
+bool SPropertyEditorAsset::OnAssetDraggedOver( const UObject* InObject, FText& OutReason ) const
 {
 	if (CanEdit() && InObject != nullptr && InObject->IsA(ObjectClass))
 	{
@@ -1219,7 +1219,7 @@ bool SPropertyEditorAsset::OnAssetDraggedOver( const UObject* InObject ) const
 		{
 			if (CanSetBasedOnCustomClasses(AssetData))
 			{
-				return CanSetBasedOnAssetReferenceFilter(AssetData);
+				return CanSetBasedOnAssetReferenceFilter(AssetData, &OutReason);
 			}
 		}
 	}
diff --git a/Engine/Source/Editor/PropertyEditor/Private/UserInterface/PropertyEditor/SPropertyEditorAsset.h b/Engine/Source/Editor/PropertyEditor/Private/UserInterface/PropertyEditor/SPropertyEditorAsset.h
--- a/Engine/Source/Editor/PropertyEditor/Private/UserInterface/PropertyEditor/SPropertyEditorAsset.h
+++ b/Engine/Source/Editor/PropertyEditor/Private/UserInterface/PropertyEditor/SPropertyEditorAsset.h
@@ -268,9 +268,10 @@ private:
 	/** 
 	 * Delegate used to check whether we can drop an object on this widget.
 	 * @param	InObject	The object we are dragging
+	 * @param	OutReason	When returning false, the reason it was not allowed
 	 * @returns true if the object can be dropped
 	 */
-	bool OnAssetDraggedOver( const UObject* InObject ) const;
+	bool OnAssetDraggedOver( const UObject* InObject, FText& OutReason ) const;
 
 	/** 
 	 * Delegate handling dropping an object on this widget
diff --git a/Engine/Source/Editor/UnrealEd/Private/LevelEditorViewport.cpp b/Engine/Source/Editor/UnrealEd/Private/LevelEditorViewport.cpp
--- a/Engine/Source/Editor/UnrealEd/Private/LevelEditorViewport.cpp
+++ b/Engine/Source/Editor/UnrealEd/Private/LevelEditorViewport.cpp
@@ -1302,11 +1302,30 @@ FDropQuery FLevelEditorViewportClient::CanDropObjectsAtCoordinates(int32 MouseX,
 {
 	FDropQuery Result;
 
-	if ( !ObjectTools::IsAssetValidForPlacing( GetWorld(), AssetData.ObjectPath.ToString() ) )
+	UWorld* CurrentWorld = GetWorld();
+	if ( !ObjectTools::IsAssetValidForPlacing(CurrentWorld, AssetData.ObjectPath.ToString() ) )
 	{
 		return Result;
 	}
 
+	if (CurrentWorld)
+	{
+		FAssetReferenceFilterContext AssetReferenceFilterContext;
+		AssetReferenceFilterContext.ReferencingAssets.Add(FAssetData(CurrentWorld));
+
+		TSharedPtr<IAssetReferenceFilter> AssetReferenceFilter = GEditor->MakeAssetReferenceFilter(AssetReferenceFilterContext);
+		if (AssetReferenceFilter.IsValid())
+		{
+			FText FailureReason;
+			if (!AssetReferenceFilter->PassesFilter(AssetData, &FailureReason))
+			{
+				Result.bCanDrop = false;
+				Result.HintText = FailureReason;
+				return Result;
+			}
+		}
+	}
+
 	UObject* AssetObj = AssetData.GetAsset();
 	UClass* ClassObj = Cast<UClass>( AssetObj );
 
