Fix combobox dropdowns don't dissapear when losing focus.

diff --git a/Engine/Source/Editor/EditorWidgets/Private/SSearchableComboBox.cpp b/Engine/Source/Editor/EditorWidgets/Private/SSearchableComboBox.cpp
--- a/Engine/Source/Editor/EditorWidgets/Private/SSearchableComboBox.cpp
+++ b/Engine/Source/Editor/EditorWidgets/Private/SSearchableComboBox.cpp
@@ -82,7 +82,7 @@ void SSearchableComboBox::Construct(const FArguments& InArgs)
 		.IsFocusable(true)
 		);
 
-	if (InArgs._SearchVisibility.Get() != EVisibility::Collapsed)
+	if (SearchField->GetVisibility() == EVisibility::Visible)
 	{
 		SetMenuContentWidgetToFocus(SearchField);
 	}
