Fix text property customization to work with new UE5 style

diff --git a/Engine/Source/Editor/EditorWidgets/Private/STextPropertyEditableTextBox.cpp b/Engine/Source/Editor/EditorWidgets/Private/STextPropertyEditableTextBox.cpp
--- a/Engine/Source/Editor/EditorWidgets/Private/STextPropertyEditableTextBox.cpp
+++ b/Engine/Source/Editor/EditorWidgets/Private/STextPropertyEditableTextBox.cpp
@@ -611,7 +611,7 @@ void STextPropertyEditableTextBox::Construct(const FArguments& InArgs, const TSh
 		[
 			SAssignNew(HorizontalBox, SHorizontalBox)
 			+SHorizontalBox::Slot()
-			.FillWidth(1.0f)
+			.VAlign(VAlign_Center)
 			[
 				SNew(SBox)
 				.MinDesiredWidth(InArgs._MinDesiredWidth)
@@ -646,6 +646,13 @@ void STextPropertyEditableTextBox::Construct(const FArguments& InArgs, const TSh
 			.ButtonStyle(FEditorStyle::Get(), "HoverHintOnly")
 			.ForegroundColor(FSlateColor::UseForeground())
 			.ToolTipText(LOCTEXT("AdvancedTextSettingsComboToolTip", "Edit advanced text settings."))
+			.HasDownArrow(false)
+			.ButtonContent()
+			[
+				SNew(SImage)
+				.Image(FAppStyle::Get().GetBrush("LocalizationDashboard.MenuIcon"))
+				.ColorAndOpacity(FSlateColor::UseForeground())
+			]
 			.MenuContent()
 			[
 				SNew(SBox)
