Details name area textbox changed to be an SInlineEditableTextBlock.

Renaming cleaned up to avoid duplicating code in the single and multiple selection cases.

diff --git a/Engine/Source/Editor/EditorStyle/Private/StarshipStyle.cpp b/Engine/Source/Editor/EditorStyle/Private/StarshipStyle.cpp
--- a/Engine/Source/Editor/EditorStyle/Private/StarshipStyle.cpp
+++ b/Engine/Source/Editor/EditorStyle/Private/StarshipStyle.cpp
@@ -3123,6 +3123,14 @@ void FStarshipEditorStyle::FStyle::SetupPropertyEditorStyles()
 
 		Set( "PropertyWindow.WindowBorder", new BOX_BRUSH( "Common/GroupBorder", FMargin(4.0f/16.0f) ) );
 
+		FInlineEditableTextBlockStyle NameStyle(FCoreStyle::Get().GetWidgetStyle<FInlineEditableTextBlockStyle>("InlineEditableTextBlockStyle"));
+		NameStyle.EditableTextBoxStyle.SetFont(DEFAULT_FONT("Regular", 11))
+			.SetForegroundColor(FSlateColor(EStyleColor::White));
+		NameStyle.TextStyle.SetFont(DEFAULT_FONT("Regular", 11))
+			.SetColorAndOpacity(FSlateColor(EStyleColor::White));
+
+		Set( "DetailsView.NameTextBlockStyle",  NameStyle );
+
 		Set( "DetailsView.NameChangeCommitted", new BOX_BRUSH( "Common/EditableTextSelectionBackground", FMargin(4.f/16.f) ) );
 		Set( "DetailsView.HyperlinkStyle", FTextBlockStyle(NormalText) .SetFont( DEFAULT_FONT( "Regular", 8 ) ) );
 
diff --git a/Engine/Source/Editor/EditorWidgets/Private/SObjectNameEditableTextBox.cpp b/Engine/Source/Editor/EditorWidgets/Private/SObjectNameEditableTextBox.cpp
--- a/Engine/Source/Editor/EditorWidgets/Private/SObjectNameEditableTextBox.cpp
+++ b/Engine/Source/Editor/EditorWidgets/Private/SObjectNameEditableTextBox.cpp
@@ -3,7 +3,7 @@
 #include "SObjectNameEditableTextBox.h"
 #include "Rendering/DrawElements.h"
 #include "Framework/Application/SlateApplication.h"
-#include "Widgets/Input/SEditableTextBox.h"
+#include "Widgets/Text/SInlineEditableTextBlock.h"
 #include "EditorStyleSet.h"
 #include "GameFramework/Actor.h"
 #include "ScopedTransaction.h"
@@ -27,18 +27,16 @@ void SObjectNameEditableTextBox::Construct( const FArguments& InArgs )
 	Objects = InArgs._Objects;
 
 	ChildSlot
-		[
-			SAssignNew(TextBox, SEditableTextBox)
-			.Text( this, &SObjectNameEditableTextBox::GetNameText )
-			.ToolTipText( this, &SObjectNameEditableTextBox::GetNameTooltipText )
-			.Visibility( this, &SObjectNameEditableTextBox::GetNameVisibility )
-			.HintText( this, &SObjectNameEditableTextBox::GetNameHintText )
-			.OnTextCommitted( this, &SObjectNameEditableTextBox::OnNameTextCommitted )
-			.IsReadOnly( this, &SObjectNameEditableTextBox::CannotEditNameText )
-			.SelectAllTextWhenFocused( this, &SObjectNameEditableTextBox::CanEditNameText )
-			.OnTextChanged(this, &SObjectNameEditableTextBox::OnTextChanged )
-			.RevertTextOnEscape( true )
-		];
+	[
+		SAssignNew(TextBox, SInlineEditableTextBlock)
+		.Style(FAppStyle::Get(), "DetailsView.NameTextBlockStyle")
+		.Text( this, &SObjectNameEditableTextBox::GetNameText )
+		.ToolTipText( this, &SObjectNameEditableTextBox::GetNameTooltipText )
+		.Visibility( this, &SObjectNameEditableTextBox::GetNameVisibility )
+		.OnTextCommitted( this, &SObjectNameEditableTextBox::OnNameTextCommitted )
+		.IsReadOnly( this, &SObjectNameEditableTextBox::IsReadOnly )
+		.OnVerifyTextChanged_Static(&FActorEditorUtils::ValidateActorName)
+	];
 }
 
 EActiveTimerReturnType SObjectNameEditableTextBox::UpdateHighlightSpringState( double InCurrentTime, float InDeltaTime )
@@ -158,7 +156,7 @@ FText SObjectNameEditableTextBox::GetNameTooltipText() const
 	}
 	else if (Objects.Num() == 1 && Objects[0].IsValid())
 	{
-		if (CanEditNameText())
+		if (!IsReadOnly())
 		{
 			Result = FText::Format(LOCTEXT("EditableActorLabel_ActorTooltipFmt", "Rename the selected {0}"), FText::FromString(Objects[0].Get()->GetClass()->GetName()));
 		}
@@ -173,7 +171,7 @@ FText SObjectNameEditableTextBox::GetNameTooltipText() const
 	}
 	else if (Objects.Num() > 1)
 	{
-		if (CanEditNameText())
+		if (!IsReadOnly())
 		{
 			Result = LOCTEXT("EditableActorLabel_MultiActorTooltip", "Rename multiple selected actors at once");
 		}
@@ -190,7 +188,7 @@ EVisibility SObjectNameEditableTextBox::GetNameVisibility() const
 {
 	if (Objects.Num() == 1 && Objects[0].IsValid())
 	{
-		if (CanEditNameText())
+		if (!IsReadOnly())
 		{
 			return EVisibility::Visible;
 		}
@@ -205,7 +203,7 @@ EVisibility SObjectNameEditableTextBox::GetNameVisibility() const
 	}
 	else if (Objects.Num() > 1)
 	{
-		if (CanEditNameText())
+		if (!IsReadOnly())
 		{
 			return EVisibility::Visible;
 		}
@@ -216,135 +214,105 @@ EVisibility SObjectNameEditableTextBox::GetNameVisibility() const
 	}
 
 	return EVisibility::Collapsed;
-}
+} 
 
-FText SObjectNameEditableTextBox::GetNameHintText() const
-{
-	FText Result;
 
-	if (Objects.Num() == 0)
+static bool RenameActor(TWeakObjectPtr<UObject> Object, FString Name)
+{
+	// Apply the change to the selected actor
+	if (!Object.IsValid() || !Object.Get()->IsA(AActor::StaticClass()))
 	{
-		Result = LOCTEXT("EditableActorLabel_NoObjectsHint", "<Nothing Selected>");
+		return false;
 	}
-	else if (Objects.Num() == 1 && Objects[0].IsValid())
+
+	AActor* Actor = Cast<AActor>(Object.Get());
+	if (!Actor->IsActorLabelEditable())
 	{
-		Result = FText::Format(LOCTEXT("EditableActorLabel_MultiObjectsHint_SameType", "<Selected {0}>"), FText::FromName(Objects[0].Get()->GetClass()->GetFName()));
+		return false;
 	}
-	else if (Objects.Num() > 1)
+
+	if (Actor->GetActorLabel() == Name)
 	{
-		Result = LOCTEXT("EditableActorLabel_MultiObjectsHint_DifferentTypes", "<Selected Objects>");
+		return false;
 	}
 
-	return Result;
+	FActorLabelUtilities::RenameExistingActor(Actor, Name);
+	return true;
 }
 
 void SObjectNameEditableTextBox::OnNameTextCommitted(const FText& NewText, ETextCommit::Type InTextCommit)
 {
 	// Don't apply the change if the TextCommit type is OnCleared - this will only be the case if the keyboard focus was cleared due to
 	// Enter being pressed, in which case we will already have been here once with a TextCommit type of OnEnter.
-	FText ErrorMessage;
-	if (InTextCommit != ETextCommit::OnCleared && FActorEditorUtils::ValidateActorName(NewText, ErrorMessage))
+	if (InTextCommit == ETextCommit::OnCleared)
 	{
-		FText TrimmedText = FText::TrimPrecedingAndTrailing(NewText);
+		return;
+	}
 
-		if (!TrimmedText.IsEmpty())
-		{
-			if (Objects.Num() == 1)
-			{
-				// Apply the change to the selected actor
-				AActor* Actor = Cast<AActor>(Objects[0].Get());
-				if(Actor != NULL)
-				{
-					const FScopedTransaction Transaction( LOCTEXT("RenameActorTransaction", "Rename Actor") );
-
-					if (Actor->IsActorLabelEditable())
-					{
-						FActorLabelUtilities::RenameExistingActor(Actor, TrimmedText.ToString());
-						LastCommittedTime = FSlateApplication::Get().GetCurrentTime();
-						RegisterActiveTimer( 0.f, FWidgetActiveTimerDelegate::CreateSP( this, &SObjectNameEditableTextBox::UpdateHighlightSpringState ) );
-					}
-				}
-			}
-			else if (Objects.Num() > 1)
-			{
-				const FScopedTransaction Transaction( LOCTEXT("RenameActorsTransaction", "Rename Multiple Actors") );
 
-				UserSetCommonName = TrimmedText.ToString();
+	FText TrimmedText = FText::TrimPrecedingAndTrailing(NewText);
+	if (TrimmedText.IsEmpty())
+	{
+		return;
+	}
 
-				for (int32 i=0; i<Objects.Num(); i++)
-				{
-					// Apply the change to the selected actor
-					if(Objects[i].IsValid() && Objects[i].Get()->IsA(AActor::StaticClass()))
-					{
-						AActor* Actor = (AActor*)Objects[i].Get();
-						if (Actor->IsActorLabelEditable())
-						{
-							FActorLabelUtilities::RenameExistingActor(Actor, TrimmedText.ToString());
-							LastCommittedTime = FSlateApplication::Get().GetCurrentTime();
-							RegisterActiveTimer( 0.f, FWidgetActiveTimerDelegate::CreateSP( this, &SObjectNameEditableTextBox::UpdateHighlightSpringState ) );
-						}
-					}
-				}
-			}
-		}
-		// Remove ourselves from the window focus so we don't get automatically reselected when scrolling around the context menu.
-		TSharedPtr< SWindow > ParentWindow = FSlateApplication::Get().FindWidgetWindow( SharedThis(this) );
-		if( ParentWindow.IsValid() )
+	UserSetCommonName = TrimmedText.ToString();
+
+	const FScopedTransaction Transaction( LOCTEXT("RenameActorsTransaction", "Rename Multiple Actors") );
+
+	bool bChanged = false;
+	for (TWeakObjectPtr<UObject> Object : Objects)
+	{
+		if (RenameActor(Object, UserSetCommonName))
 		{
-			ParentWindow->SetWidgetToFocusOnActivate( NULL );
+			bChanged = true;
 		}
 	}
 
-	// Clear Error 
-	TextBox->SetError(FText::GetEmpty());
-}
-
-void SObjectNameEditableTextBox::OnTextChanged( const FText& InLabel )
-{
-	FText OutErrorMessage;
-
-	if(!FActorEditorUtils::ValidateActorName(InLabel, OutErrorMessage))
+	if (bChanged)
 	{
-		TextBox->SetError(OutErrorMessage);
+		LastCommittedTime = FSlateApplication::Get().GetCurrentTime();
+		RegisterActiveTimer( 0.f, FWidgetActiveTimerDelegate::CreateSP( this, &SObjectNameEditableTextBox::UpdateHighlightSpringState ) );
 	}
-	else
+
+	// Remove ourselves from the window focus so we don't get automatically reselected when scrolling around the context menu.
+	TSharedPtr< SWindow > ParentWindow = FSlateApplication::Get().FindWidgetWindow( SharedThis(this) );
+	if (ParentWindow.IsValid())
 	{
-		TextBox->SetError( FText::GetEmpty() );
+		ParentWindow->SetWidgetToFocusOnActivate(nullptr);
 	}
 }
 
-bool SObjectNameEditableTextBox::CanEditNameText() const
+bool SObjectNameEditableTextBox::IsReadOnly() const
 {
-	bool Result = false;
+	if (Objects.Num() == 0)
+	{
+		// can't edit if nothing is selected
+		return true;
+	}
 
-	if (Objects.Num() > 0)
+	for (TWeakObjectPtr<UObject> Object : Objects)
 	{
-		Result = true;
-		for (int32 i=0; i<Objects.Num(); i++)
+		if (Object.IsValid())
 		{
-			if(Objects[i].IsValid())
+			if (Object.Get()->IsA(AActor::StaticClass()))
 			{
-				if (Objects[i].Get()->IsA(AActor::StaticClass()))
+				AActor* Actor = (AActor*) Object.Get();
+				if (!Actor->IsActorLabelEditable())
 				{
-					AActor* Actor = (AActor*)Objects[i].Get();
-					if (!Actor->IsActorLabelEditable())
-					{
-						// can't edit the name when a non-editable actor is selected
-						Result = false;
-						break;
-					}
-				}
-				else
-				{
-					// can't edit the name when a non-actor is selected
-					Result = false;
-					break;
+					// can't edit the name when a non-editable actor is selected
+					return true;
 				}
 			}
+			else
+			{
+				// can't edit the name when a non-actor is selected
+				return true;
+			}
 		}
 	}
 
-	return Result;
+	return false;
 }
 
 FString SObjectNameEditableTextBox::GetObjectDisplayName(TWeakObjectPtr<UObject> Object)
diff --git a/Engine/Source/Editor/EditorWidgets/Private/SObjectNameEditableTextBox.h b/Engine/Source/Editor/EditorWidgets/Private/SObjectNameEditableTextBox.h
--- a/Engine/Source/Editor/EditorWidgets/Private/SObjectNameEditableTextBox.h
+++ b/Engine/Source/Editor/EditorWidgets/Private/SObjectNameEditableTextBox.h
@@ -11,7 +11,7 @@
 
 class FPaintArgs;
 class FSlateWindowElementList;
-class SEditableTextBox;
+class SInlineEditableTextBlock;
 
 /** Widget wraps an editable text box for viewing the names of objects or editing the labels of actors */
 class SObjectNameEditableTextBox : public IObjectNameEditableTextBox
@@ -49,20 +49,11 @@ private:
 	/** Getter for the ToolTipText attribute of the editable text inside this widget */
 	FText GetNameTooltipText() const;
 
-	/** Getter for the HintText attribute of the editable text inside this widget */
-	FText GetNameHintText() const;
-
 	/** Getter for the OnTextCommitted event of the editable text inside this widget */
 	void OnNameTextCommitted(const FText& NewText, ETextCommit::Type InTextCommit);
 
 	/** Getter for the IsReadOnly attribute of the editable text inside this widget */
-	bool CanEditNameText() const;
-
-	/** Getter for the SelectAllTextWhenFocused attribute of the editable text inside this widget */
-	bool CannotEditNameText() const { return !CanEditNameText(); }
-
-	/** Callback to verify a text change */
-	void OnTextChanged( const FText& InLabel );
+	bool IsReadOnly() const;
 
 	/** Helper class the get the object name or the actor label if an object is an actor */
 	static FString GetObjectDisplayName(TWeakObjectPtr<UObject> Object);
@@ -101,7 +92,7 @@ private:
 	double LastCommittedTime;
 
 	/** The text box used to edit object names */ 
-	TSharedPtr< SEditableTextBox > TextBox;
+	TSharedPtr< SInlineEditableTextBlock > TextBox;
 
 	// Temp flag to trigger a highlight spring update in the passive tick (because that's where the geometry is)
 	bool bUpdateHighlightSpring;
diff --git a/Engine/Source/Editor/PropertyEditor/Private/SDetailNameArea.cpp b/Engine/Source/Editor/PropertyEditor/Private/SDetailNameArea.cpp
--- a/Engine/Source/Editor/PropertyEditor/Private/SDetailNameArea.cpp
+++ b/Engine/Source/Editor/PropertyEditor/Private/SDetailNameArea.cpp
@@ -139,7 +139,7 @@ TSharedRef< SWidget > SDetailNameArea::BuildObjectNameArea( const TArray< TWeakO
 
 			ObjectNameArea->AddSlot()
 				.FillWidth(1.0)
-				.Padding(0, 0, 3, 0)
+				.Padding(5, 0, 3, 0)
 				[
 					SNew(SBox)
 					.WidthOverride(200.0f)
