Renamed SEnumCombobox.h to SEnumCombo.h to remove capialization issue

diff --git a/Engine/Plugins/Experimental/ControlRig/Source/ControlRigEditor/Private/ControlRigElementDetails.cpp b/Engine/Plugins/Experimental/ControlRig/Source/ControlRigEditor/Private/ControlRigElementDetails.cpp
--- a/Engine/Plugins/Experimental/ControlRig/Source/ControlRigEditor/Private/ControlRigElementDetails.cpp
+++ b/Engine/Plugins/Experimental/ControlRig/Source/ControlRigEditor/Private/ControlRigElementDetails.cpp
@@ -18,7 +18,7 @@
 #include "ControlRigBlueprint.h"
 #include "Graph/ControlRigGraph.h"
 #include "PropertyCustomizationHelpers.h"
-#include "SEnumCombobox.h"
+#include "SEnumCombo.h"
 #include "ControlRig/Private/Units/Execution/RigUnit_BeginExecution.h"
 
 #define LOCTEXT_NAMESPACE "ControlRigElementDetails"
diff --git a/Engine/Plugins/Experimental/ControlRig/Source/ControlRigEditor/Private/EditMode/ControlRigControlsProxy.cpp b/Engine/Plugins/Experimental/ControlRig/Source/ControlRigEditor/Private/EditMode/ControlRigControlsProxy.cpp
--- a/Engine/Plugins/Experimental/ControlRig/Source/ControlRigEditor/Private/EditMode/ControlRigControlsProxy.cpp
+++ b/Engine/Plugins/Experimental/ControlRig/Source/ControlRigEditor/Private/EditMode/ControlRigControlsProxy.cpp
@@ -16,7 +16,7 @@
 #include "DetailLayoutBuilder.h"
 #include "DetailCategoryBuilder.h"
 #include "DetailWidgetRow.h"
-#include "SEnumCombobox.h"
+#include "SEnumCombo.h"
 
 void FControlRigEnumControlProxyValueDetails::CustomizeHeader(TSharedRef<IPropertyHandle> InStructPropertyHandle, FDetailWidgetRow& HeaderRow, IPropertyTypeCustomizationUtils& StructCustomizationUtils)
 {
diff --git a/Engine/Plugins/Experimental/ControlRig/Source/ControlRigEditor/Private/Editor/ControlRigEditor.cpp b/Engine/Plugins/Experimental/ControlRig/Source/ControlRigEditor/Private/Editor/ControlRigEditor.cpp
--- a/Engine/Plugins/Experimental/ControlRig/Source/ControlRigEditor/Private/Editor/ControlRigEditor.cpp
+++ b/Engine/Plugins/Experimental/ControlRig/Source/ControlRigEditor/Private/Editor/ControlRigEditor.cpp
@@ -7,7 +7,7 @@
 #include "SBlueprintEditorToolbar.h"
 #include "ControlRigEditorMode.h"
 #include "SKismetInspector.h"
-#include "SEnumCombobox.h"
+#include "SEnumCombo.h"
 #include "Framework/Commands/GenericCommands.h"
 #include "Editor.h"
 #include "Graph/ControlRigGraph.h"
diff --git a/Engine/Plugins/FX/Niagara/Source/NiagaraEditor/Private/TypeEditorUtilities/NiagaraEnumTypeEditorUtilities.cpp b/Engine/Plugins/FX/Niagara/Source/NiagaraEditor/Private/TypeEditorUtilities/NiagaraEnumTypeEditorUtilities.cpp
--- a/Engine/Plugins/FX/Niagara/Source/NiagaraEditor/Private/TypeEditorUtilities/NiagaraEnumTypeEditorUtilities.cpp
+++ b/Engine/Plugins/FX/Niagara/Source/NiagaraEditor/Private/TypeEditorUtilities/NiagaraEnumTypeEditorUtilities.cpp
@@ -7,7 +7,7 @@
 #include "NiagaraEditorCommon.h"
 
 #include "Widgets/Input/SSpinBox.h"
-#include "EditorWidgets/Public/SEnumCombobox.h"
+#include "EditorWidgets/Public/SEnumCombo.h"
 #include "EditorStyleSet.h"
 
 class SNiagaraEnumParameterEditor : public SNiagaraParameterEditor
diff --git a/Engine/Plugins/FX/Niagara/Source/NiagaraEditor/Private/ViewModels/NiagaraParameterNameViewModel.cpp b/Engine/Plugins/FX/Niagara/Source/NiagaraEditor/Private/ViewModels/NiagaraParameterNameViewModel.cpp
--- a/Engine/Plugins/FX/Niagara/Source/NiagaraEditor/Private/ViewModels/NiagaraParameterNameViewModel.cpp
+++ b/Engine/Plugins/FX/Niagara/Source/NiagaraEditor/Private/ViewModels/NiagaraParameterNameViewModel.cpp
@@ -1,7 +1,7 @@
 // Copyright Epic Games, Inc. All Rights Reserved.
 
 #include "NiagaraParameterNameViewModel.h"
-#include "EditorWidgets/Public/SEnumCombobox.h"
+#include "EditorWidgets/Public/SEnumCombo.h"
 #include "Widgets/Text/SInlineEditableTextBlock.h"
 #include "SGraphActionMenu.h"
 #include "NiagaraActions.h"
diff --git a/Engine/Source/Editor/DetailCustomizations/Private/FbxImportUIDetails.cpp b/Engine/Source/Editor/DetailCustomizations/Private/FbxImportUIDetails.cpp
--- a/Engine/Source/Editor/DetailCustomizations/Private/FbxImportUIDetails.cpp
+++ b/Engine/Source/Editor/DetailCustomizations/Private/FbxImportUIDetails.cpp
@@ -18,7 +18,7 @@
 #include "Misc/Attribute.h"
 #include "Misc/Guid.h"
 #include "PropertyHandle.h"
-#include "SEnumCombobox.h"
+#include "SEnumCombo.h"
 #include "Templates/Tuple.h"
 #include "Widgets/DeclarativeSyntaxSupport.h"
 #include "Widgets/Images/SImage.h"
diff --git a/Engine/Source/Editor/EditorWidgets/Private/SEnumCombo.cpp b/Engine/Source/Editor/EditorWidgets/Private/SEnumCombo.cpp
new file mode 100644
--- /dev/null
+++ b/Engine/Source/Editor/EditorWidgets/Private/SEnumCombo.cpp
@@ -0,0 +1,92 @@
+// Copyright Epic Games, Inc. All Rights Reserved.
+
+#include "SEnumCombo.h"
+#include "EditorStyleSet.h"
+#include "Types/SlateEnums.h"
+#include "Widgets/SToolTip.h"
+
+void SEnumComboBox::Construct(const FArguments& InArgs, const UEnum* InEnum)
+{
+	Enum = InEnum;
+	CurrentValue = InArgs._CurrentValue;
+	check(CurrentValue.IsBound());
+	OnEnumSelectionChangedDelegate = InArgs._OnEnumSelectionChanged;
+	OnGetToolTipForValue = InArgs._OnGetToolTipForValue;
+	Font = InArgs._Font;
+	bUpdatingSelectionInternally = false;
+
+	for (int32 i = 0; i < Enum->NumEnums() - 1; i++)
+	{
+		if (Enum->HasMetaData(TEXT("Hidden"), i) == false)
+		{
+			VisibleEnumNameIndices.Add(MakeShareable(new int32(i)));
+		}
+	}
+
+	SComboBox::Construct(SComboBox<TSharedPtr<int32>>::FArguments()
+		.ButtonStyle(InArgs._ButtonStyle)
+		.OptionsSource(&VisibleEnumNameIndices)
+		.OnGenerateWidget(this, &SEnumComboBox::OnGenerateWidget)
+		.OnSelectionChanged(this, &SEnumComboBox::OnComboSelectionChanged)
+		.OnComboBoxOpening(this, &SEnumComboBox::OnComboMenuOpening)
+		.ContentPadding(InArgs._ContentPadding)
+		[
+			SNew(STextBlock)
+			.Font(Font)
+			.Text(this, &SEnumComboBox::GetCurrentValueText)
+			.ToolTipText(this, &SEnumComboBox::GetCurrentValueTooltip)
+		]);
+}
+
+
+FText SEnumComboBox::GetCurrentValueText() const
+{
+	int32 ValueNameIndex = Enum->GetIndexByValue(CurrentValue.Get());
+	return Enum->GetDisplayNameTextByIndex(ValueNameIndex);
+}
+
+FText SEnumComboBox::GetCurrentValueTooltip() const
+{
+	int32 ValueNameIndex = Enum->GetIndexByValue(CurrentValue.Get());
+	if (OnGetToolTipForValue.IsBound())
+	{
+		return OnGetToolTipForValue.Execute(ValueNameIndex);
+	}
+	return Enum->GetToolTipTextByIndex(ValueNameIndex);
+}
+
+TSharedRef<SWidget> SEnumComboBox::OnGenerateWidget(TSharedPtr<int32> InItem)
+{
+	return SNew(STextBlock)
+		.Font(Font)
+		.Text(Enum->GetDisplayNameTextByIndex(*InItem))
+		.ToolTipText(Enum->GetToolTipTextByIndex(*InItem));
+}
+
+void SEnumComboBox::OnComboSelectionChanged(TSharedPtr<int32> InSelectedItem, ESelectInfo::Type SelectInfo)
+{
+	if (bUpdatingSelectionInternally == false)
+	{
+		OnEnumSelectionChangedDelegate.ExecuteIfBound(*InSelectedItem, SelectInfo);
+	}
+}
+
+void SEnumComboBox::OnComboMenuOpening()
+{
+	int32 CurrentNameIndex = Enum->GetIndexByValue(CurrentValue.Get());
+	TSharedPtr<int32> FoundNameIndexItem;
+	for (int32 i = 0; i < VisibleEnumNameIndices.Num(); i++)
+	{
+		if (*VisibleEnumNameIndices[i] == CurrentNameIndex)
+		{
+			FoundNameIndexItem = VisibleEnumNameIndices[i];
+			break;
+		}
+	}
+	if (FoundNameIndexItem.IsValid())
+	{
+		bUpdatingSelectionInternally = true;
+		SetSelectedItem(FoundNameIndexItem);
+		bUpdatingSelectionInternally = false;
+	}
+}
diff --git a/Engine/Source/Editor/EditorWidgets/Private/SEnumComboBox.cpp b/Engine/Source/Editor/EditorWidgets/Private/SEnumComboBox.cpp
deleted file mode 100644
diff --git a/Engine/Source/Editor/EditorWidgets/Public/SEnumCombo.h b/Engine/Source/Editor/EditorWidgets/Public/SEnumCombo.h
new file mode 100644
--- /dev/null
+++ b/Engine/Source/Editor/EditorWidgets/Public/SEnumCombo.h
@@ -0,0 +1,53 @@
+// Copyright Epic Games, Inc. All Rights Reserved.
+
+#pragma once
+
+#include "CoreMinimal.h"
+#include "Widgets/Input/SComboBox.h"
+
+class EDITORWIDGETS_API SEnumComboBox : public SComboBox<TSharedPtr<int32>>
+{
+public:
+	DECLARE_DELEGATE_TwoParams(FOnEnumSelectionChanged, int32, ESelectInfo::Type);
+	DECLARE_DELEGATE_RetVal_OneParam(FText, FOnGetToolTipForValue, int32 /* Value */)
+
+public:
+	SLATE_BEGIN_ARGS(SEnumComboBox)
+		: _CurrentValue()
+		, _ContentPadding(FMargin(4.0, 2.0))
+		, _ButtonStyle(nullptr)
+	{}
+
+		SLATE_ATTRIBUTE(int32, CurrentValue)
+		SLATE_ATTRIBUTE(FMargin, ContentPadding)
+		SLATE_ATTRIBUTE(FSlateFontInfo, Font)
+		SLATE_EVENT(FOnEnumSelectionChanged, OnEnumSelectionChanged)
+		SLATE_EVENT(FOnGetToolTipForValue, OnGetToolTipForValue)
+		SLATE_STYLE_ARGUMENT(FButtonStyle, ButtonStyle)
+	SLATE_END_ARGS()
+
+	void Construct(const FArguments& InArgs, const UEnum* InEnum);
+
+private:
+	FText GetCurrentValueText() const;
+	FText GetCurrentValueTooltip() const;
+	bool GetValueIsEnabled(TWeakPtr<int32> InValueWeak) const;
+	TSharedRef<SWidget> OnGenerateWidget(TSharedPtr<int32> InItem);
+	void OnComboSelectionChanged(TSharedPtr<int32> InSelectedItem, ESelectInfo::Type SelectInfo);
+	void OnComboMenuOpening();
+
+private:
+	const UEnum* Enum;
+
+	TAttribute<int32> CurrentValue;
+
+	TAttribute<FSlateFontInfo> Font;
+
+	TArray<TSharedPtr<int32>> VisibleEnumNameIndices;
+
+	bool bUpdatingSelectionInternally;
+
+	FOnGetToolTipForValue OnGetToolTipForValue;
+
+	FOnEnumSelectionChanged OnEnumSelectionChangedDelegate;
+};
diff --git a/Engine/Source/Editor/EditorWidgets/Public/SEnumCombobox.h b/Engine/Source/Editor/EditorWidgets/Public/SEnumCombobox.h
deleted file mode 100644
diff --git a/Engine/Source/Editor/MovieSceneTools/Public/MovieSceneToolHelpers.h b/Engine/Source/Editor/MovieSceneTools/Public/MovieSceneToolHelpers.h
--- a/Engine/Source/Editor/MovieSceneTools/Public/MovieSceneToolHelpers.h
+++ b/Engine/Source/Editor/MovieSceneTools/Public/MovieSceneToolHelpers.h
@@ -15,7 +15,7 @@
 #include "MovieSceneTranslator.h"
 #include "MovieSceneSpawnable.h"
 #include "MovieSceneCaptureSettings.h"
-#include "SEnumCombobox.h"
+#include "SEnumCombo.h"
 #include "Animation/AnimSequence.h"
 #include "INodeAndChannelMappings.h"
 
diff --git a/Engine/Source/Editor/UnrealEd/Private/Kismet2/DebuggerCommands.cpp b/Engine/Source/Editor/UnrealEd/Private/Kismet2/DebuggerCommands.cpp
--- a/Engine/Source/Editor/UnrealEd/Private/Kismet2/DebuggerCommands.cpp
+++ b/Engine/Source/Editor/UnrealEd/Private/Kismet2/DebuggerCommands.cpp
@@ -71,7 +71,7 @@
 
 #include "ToolMenus.h"
 #include "SBlueprintEditorToolbar.h"
-#include "SEnumCombobox.h"
+#include "SEnumCombo.h"
 #include "Dialogs/Dialogs.h"
 
 #include "Misc/DataDrivenPlatformInfoRegistry.h"
