Warning and error box now has an optional icon size and padding

diff --git a/Engine/Source/Editor/EditorWidgets/Private/SWarningOrErrorBox.cpp b/Engine/Source/Editor/EditorWidgets/Private/SWarningOrErrorBox.cpp
--- a/Engine/Source/Editor/EditorWidgets/Private/SWarningOrErrorBox.cpp
+++ b/Engine/Source/Editor/EditorWidgets/Private/SWarningOrErrorBox.cpp
@@ -7,20 +7,23 @@
 
 void SWarningOrErrorBox::Construct(const FArguments& InArgs)
 {
+	MessageStyle = InArgs._MessageStyle;
+
 	SBorder::Construct(SBorder::FArguments()
-		.Padding(16.0f)
+		.Padding(InArgs._Padding)
 		.ForegroundColor(FAppStyle::Get().GetSlateColor("Colors.White"))
-		.BorderImage(InArgs._MessageStyle == EMessageStyle::Warning ? FAppStyle::Get().GetBrush("RoundedWarning") : FAppStyle::Get().GetBrush("RoundedError"))
+		.BorderImage_Lambda([this]() { return MessageStyle.Get() == EMessageStyle::Warning ? FAppStyle::Get().GetBrush("RoundedWarning") : FAppStyle::Get().GetBrush("RoundedError"); })
 		[
 			SNew(SHorizontalBox)
 			+ SHorizontalBox::Slot()
+			.HAlign(HAlign_Left)
 			.VAlign(VAlign_Center)
 			.AutoWidth()
 			.Padding(FMargin(0.0f, 0.0f, 16.0f, 0.0f))
 			[
 				SNew(SImage)
-				.DesiredSizeOverride(FVector2D(24,24))
-				.Image(InArgs._MessageStyle == EMessageStyle::Warning ? FAppStyle::Get().GetBrush("Icons.Warning") : FAppStyle::Get().GetBrush("Icons.Error"))
+				.DesiredSizeOverride(InArgs._IconSize)
+				.Image_Lambda([this]() { return MessageStyle.Get() == EMessageStyle::Warning ? FAppStyle::Get().GetBrush("Icons.WarningWithColor") : FAppStyle::Get().GetBrush("Icons.ErrorWithColor"); })
 				.ColorAndOpacity(FSlateColor::UseForeground())
 			]
 			+SHorizontalBox::Slot()
diff --git a/Engine/Source/Editor/EditorWidgets/Public/SWarningOrErrorBox.h b/Engine/Source/Editor/EditorWidgets/Public/SWarningOrErrorBox.h
--- a/Engine/Source/Editor/EditorWidgets/Public/SWarningOrErrorBox.h
+++ b/Engine/Source/Editor/EditorWidgets/Public/SWarningOrErrorBox.h
@@ -16,10 +16,17 @@ class EDITORWIDGETS_API SWarningOrErrorBox : public SBorder
 public:
 	SLATE_BEGIN_ARGS(SWarningOrErrorBox)
 		: _MessageStyle(EMessageStyle::Warning)
+		, _Padding(16.0f)
+		, _IconSize(24,24)
 	{}
 		SLATE_ATTRIBUTE(FText, Message)
-		SLATE_ARGUMENT(EMessageStyle, MessageStyle)
+		SLATE_ATTRIBUTE(EMessageStyle, MessageStyle)
+		SLATE_ARGUMENT(FMargin, Padding)
+		SLATE_ARGUMENT(FVector2D, IconSize)
 	SLATE_END_ARGS()
 
 	void Construct(const FArguments& InArgs);
+
+private:
+	TAttribute<EMessageStyle> MessageStyle;
 };
\ No newline at end of file
