Fix PropertyComboBox clears immediately when trying to search.

diff --git a/Engine/Source/Editor/EditorWidgets/Private/SSearchableComboBox.cpp b/Engine/Source/Editor/EditorWidgets/Private/SSearchableComboBox.cpp
--- a/Engine/Source/Editor/EditorWidgets/Private/SSearchableComboBox.cpp
+++ b/Engine/Source/Editor/EditorWidgets/Private/SSearchableComboBox.cpp
@@ -200,6 +200,11 @@ void SSearchableComboBox::OnMenuOpenChanged(bool bOpen)
 
 void SSearchableComboBox::OnSelectionChanged_Internal(TSharedPtr<FString> ProposedSelection, ESelectInfo::Type SelectInfo)
 {
+	if (!ProposedSelection)
+	{
+		return;
+	}
+
 	// Ensure that the proposed selection is different from selected
 	if (ProposedSelection != SelectedItem)
 	{
