Updated SWarningOrErrorBox to allow Content

diff --git a/Engine/Source/Editor/EditorWidgets/Private/SWarningOrErrorBox.cpp b/Engine/Source/Editor/EditorWidgets/Private/SWarningOrErrorBox.cpp
--- a/Engine/Source/Editor/EditorWidgets/Private/SWarningOrErrorBox.cpp
+++ b/Engine/Source/Editor/EditorWidgets/Private/SWarningOrErrorBox.cpp
@@ -34,6 +34,15 @@ void SWarningOrErrorBox::Construct(const FArguments& InArgs)
 				.ColorAndOpacity(FAppStyle::Get().GetSlateColor("Colors.White"))
 				.AutoWrapText(true)
 			]
+
+			+SHorizontalBox::Slot()
+			.AutoWidth()
+			.VAlign(VAlign_Center)
+			.HAlign(HAlign_Right)
+			[
+				InArgs._Content.Widget
+			]
+
 		]);
 }
 
diff --git a/Engine/Source/Editor/EditorWidgets/Public/SWarningOrErrorBox.h b/Engine/Source/Editor/EditorWidgets/Public/SWarningOrErrorBox.h
--- a/Engine/Source/Editor/EditorWidgets/Public/SWarningOrErrorBox.h
+++ b/Engine/Source/Editor/EditorWidgets/Public/SWarningOrErrorBox.h
@@ -18,11 +18,15 @@ public:
 		: _MessageStyle(EMessageStyle::Warning)
 		, _Padding(16.0f)
 		, _IconSize(24,24)
+		, _Content()
 	{}
 		SLATE_ATTRIBUTE(FText, Message)
 		SLATE_ATTRIBUTE(EMessageStyle, MessageStyle)
 		SLATE_ARGUMENT(FMargin, Padding)
 		SLATE_ARGUMENT(FVector2D, IconSize)
+
+		SLATE_DEFAULT_SLOT( FArguments, Content )
+
 	SLATE_END_ARGS()
 
 	void Construct(const FArguments& InArgs);
