Fix searchable combo boxes sometimes setting focus to an invisible widget, causing the combo box menu not to close when it loses focus

diff --git a/Engine/Source/Editor/EditorWidgets/Private/SSearchableComboBox.cpp b/Engine/Source/Editor/EditorWidgets/Private/SSearchableComboBox.cpp
--- a/Engine/Source/Editor/EditorWidgets/Private/SSearchableComboBox.cpp
+++ b/Engine/Source/Editor/EditorWidgets/Private/SSearchableComboBox.cpp
@@ -26,6 +26,9 @@ void SSearchableComboBox::Construct(const FArguments& InArgs)
 
 	FilteredOptionsSource.Append(*OptionsSource);
 
+	TAttribute<EVisibility> SearchVisibility = InArgs._SearchVisibility;
+	const EVisibility CurrentSearchVisibility = SearchVisibility.Get();
+
 	TSharedRef<SWidget> ComboBoxMenuContent =
 		SNew(SBox)
 		.MaxDesiredHeight(InArgs._MaxListHeight)
@@ -39,7 +42,7 @@ void SSearchableComboBox::Construct(const FArguments& InArgs)
 				.HintText(LOCTEXT("Search", "Search"))
 				.OnTextChanged(this, &SSearchableComboBox::OnSearchTextChanged)
 				.OnTextCommitted(this, &SSearchableComboBox::OnSearchTextCommitted)
-				.Visibility(InArgs._SearchVisibility)
+				.Visibility(SearchVisibility)
 			]
 
 			+ SVerticalBox::Slot()
@@ -82,7 +85,7 @@ void SSearchableComboBox::Construct(const FArguments& InArgs)
 		.IsFocusable(true)
 		);
 
-	if (SearchField->GetVisibility() == EVisibility::Visible)
+	if (CurrentSearchVisibility == EVisibility::Visible)
 	{
 		SetMenuContentWidgetToFocus(SearchField);
 	}
