SEnumCombo: Added support for bitflags enums

diff --git a/Engine/Source/Editor/EditorWidgets/Private/SEnumCombo.cpp b/Engine/Source/Editor/EditorWidgets/Private/SEnumCombo.cpp
--- a/Engine/Source/Editor/EditorWidgets/Private/SEnumCombo.cpp
+++ b/Engine/Source/Editor/EditorWidgets/Private/SEnumCombo.cpp
@@ -4,9 +4,14 @@
 #include "EditorStyleSet.h"
 #include "Types/SlateEnums.h"
 #include "Widgets/SToolTip.h"
+#include "Framework/MultiBox/MultiBoxBuilder.h"
+
+#define LOCTEXT_NAMESPACE "EditorWidgets"
 
 void SEnumComboBox::Construct(const FArguments& InArgs, const UEnum* InEnum)
 {
+	static const FName UseEnumValuesAsMaskValuesInEditorName(TEXT("UseEnumValuesAsMaskValuesInEditor"));
+	
 	Enum = InEnum;
 	CurrentValue = InArgs._CurrentValue;
 	check(CurrentValue.IsBound());
@@ -14,22 +19,63 @@ void SEnumComboBox::Construct(const FArguments& InArgs, const UEnum* InEnum)
 	OnGetToolTipForValue = InArgs._OnGetToolTipForValue;
 	Font = InArgs._Font;
 	bUpdatingSelectionInternally = false;
+	bIsBitflagsEnum = Enum->HasMetaData(TEXT("Bitflags"));
+
+	if (bIsBitflagsEnum)
+	{
+		const bool bUseEnumValuesAsMaskValues = Enum->GetBoolMetaData(UseEnumValuesAsMaskValuesInEditorName);
+		const int32 BitmaskBitCount = sizeof(int32) << 3;
+
+		for (int32 i = 0; i < Enum->NumEnums() - 1; i++)
+		{
+			// Note: SEnumComboBox API prior to bitflags only supports 32 bit values, truncating the value here to keep the old API.
+			int32 Value = Enum->GetValueByIndex(i);
+			const bool bIsHidden = Enum->HasMetaData(TEXT("Hidden"), i);
+			if (Value >= 0 && !bIsHidden)
+			{
+				if (bUseEnumValuesAsMaskValues)
+				{
+					if (Value >= MAX_int32 || !FMath::IsPowerOfTwo(Value))
+					{
+						continue;
+					}
+				}
+				else
+				{
+					if (Value >= BitmaskBitCount)
+					{
+						continue;
+					}
+					Value = 1 << Value;
+				}
+
+				FText DisplayName = Enum->GetDisplayNameTextByIndex(i);
+				FText TooltipText = Enum->GetToolTipTextByIndex(i);
+				if (TooltipText.IsEmpty())
+				{
+					TooltipText = FText::Format(LOCTEXT("BitmaskDefaultFlagToolTipText", "Toggle {0} on/off"), DisplayName);
+				}
 
-	for (int32 i = 0; i < Enum->NumEnums() - 1; i++)
+				VisibleEnums.Emplace(i, Value, DisplayName, TooltipText);
+			}
+		}
+	}
+	else
 	{
-		if (Enum->HasMetaData(TEXT("Hidden"), i) == false)
+		for (int32 i = 0; i < Enum->NumEnums() - 1; i++)
 		{
-			VisibleEnumNameIndices.Add(MakeShareable(new int32(i)));
+			if (Enum->HasMetaData(TEXT("Hidden"), i) == false)
+			{
+				VisibleEnums.Emplace(i, Enum->GetValueByIndex(i), Enum->GetDisplayNameTextByIndex(i), Enum->GetToolTipTextByIndex(i));
+			}
 		}
 	}
 
-	SComboBox::Construct(SComboBox<TSharedPtr<int32>>::FArguments()
+	SComboButton::Construct(SComboButton::FArguments()
 		.ButtonStyle(InArgs._ButtonStyle)
-		.OptionsSource(&VisibleEnumNameIndices)
-		.OnGenerateWidget(this, &SEnumComboBox::OnGenerateWidget)
-		.OnSelectionChanged(this, &SEnumComboBox::OnComboSelectionChanged)
-		.OnComboBoxOpening(this, &SEnumComboBox::OnComboMenuOpening)
 		.ContentPadding(InArgs._ContentPadding)
+		.OnGetMenuContent(this, &SEnumComboBox::OnGetMenuContent)
+		.ButtonContent()
 		[
 			SNew(STextBlock)
 			.Font(Font)
@@ -41,13 +87,71 @@ void SEnumComboBox::Construct(const FArguments& InArgs, const UEnum* InEnum)
 
 FText SEnumComboBox::GetCurrentValueText() const
 {
-	int32 ValueNameIndex = Enum->GetIndexByValue(CurrentValue.Get());
+	if (bIsBitflagsEnum)
+	{
+		if (CurrentValue.IsSet())
+		{
+			const int32 BitmaskValue = CurrentValue.Get();
+			if (BitmaskValue != 0)
+			{
+				TArray<FText> SetFlags;
+				SetFlags.Reserve(VisibleEnums.Num());
+
+				for (const FEnumInfo& FlagInfo : VisibleEnums)
+				{
+					if ((BitmaskValue & FlagInfo.Value) != 0)
+					{
+						SetFlags.Add(FlagInfo.DisplayName);
+					}
+				}
+				if (SetFlags.Num() > 3)
+				{
+					SetFlags.SetNum(3);
+					SetFlags.Add(FText::FromString("..."));
+				}
+
+				return FText::Join(FText::FromString(" | "), SetFlags);
+			}
+			return LOCTEXT("BitmaskButtonContentNoFlagsSet", "(No Flags Set)");
+		}
+		return FText::GetEmpty();
+	}
+
+	const int32 ValueNameIndex = Enum->GetIndexByValue(CurrentValue.Get());
 	return Enum->GetDisplayNameTextByIndex(ValueNameIndex);
 }
 
 FText SEnumComboBox::GetCurrentValueTooltip() const
 {
-	int32 ValueNameIndex = Enum->GetIndexByValue(CurrentValue.Get());
+	if (bIsBitflagsEnum)
+	{
+		const int32 BitmaskValue = CurrentValue.Get();
+		if (BitmaskValue != 0)
+		{
+			TArray<FText> SetFlags;
+			SetFlags.Reserve(VisibleEnums.Num());
+
+			for (const FEnumInfo& FlagInfo : VisibleEnums)
+			{
+				if ((BitmaskValue & FlagInfo.Value) != 0)
+				{
+					if (OnGetToolTipForValue.IsBound())
+					{
+						SetFlags.Add(OnGetToolTipForValue.Execute(FlagInfo.Index));
+					}
+					else
+					{
+						SetFlags.Add(FlagInfo.DisplayName);
+					}
+				}
+			}
+
+			return FText::Join(FText::FromString(" | "), SetFlags);
+		}
+		return FText::GetEmpty();
+	}
+	
+	const int32 ValueNameIndex = Enum->GetIndexByValue(CurrentValue.Get());
 	if (OnGetToolTipForValue.IsBound())
 	{
 		return OnGetToolTipForValue.Execute(ValueNameIndex);
@@ -55,38 +159,42 @@ FText SEnumComboBox::GetCurrentValueTooltip() const
 	return Enum->GetToolTipTextByIndex(ValueNameIndex);
 }
 
-TSharedRef<SWidget> SEnumComboBox::OnGenerateWidget(TSharedPtr<int32> InItem)
+TSharedRef<SWidget> SEnumComboBox::OnGetMenuContent()
 {
-	return SNew(STextBlock)
-		.Font(Font)
-		.Text(Enum->GetDisplayNameTextByIndex(*InItem))
-		.ToolTipText(Enum->GetToolTipTextByIndex(*InItem));
-}
+	const bool bCloseAfterSelection = !bIsBitflagsEnum;
+	FMenuBuilder MenuBuilder(bCloseAfterSelection, nullptr);
 
-void SEnumComboBox::OnComboSelectionChanged(TSharedPtr<int32> InSelectedItem, ESelectInfo::Type SelectInfo)
-{
-	if (bUpdatingSelectionInternally == false)
+	for (const FEnumInfo& FlagInfo : VisibleEnums)
 	{
-		OnEnumSelectionChangedDelegate.ExecuteIfBound(*InSelectedItem, SelectInfo);
+		MenuBuilder.AddMenuEntry(
+			FlagInfo.DisplayName,
+			FlagInfo.TooltipText,
+			FSlateIcon(),
+			FUIAction
+			(
+				FExecuteAction::CreateLambda([this, FlagInfo]()
+				{
+					if (bIsBitflagsEnum)
+					{
+						// Toggle value
+						const int32 Value = CurrentValue.Get() ^ FlagInfo.Value;
+						OnEnumSelectionChangedDelegate.ExecuteIfBound(Value, ESelectInfo::Direct);
+					}
+					else
+					{
+						// Set value
+						OnEnumSelectionChangedDelegate.ExecuteIfBound(FlagInfo.Value, ESelectInfo::OnMouseClick);
+					}
+				}),
+				FCanExecuteAction(),
+				FIsActionChecked::CreateLambda([this, FlagInfo]() -> bool
+				{
+					return (CurrentValue.Get() & FlagInfo.Value) != 0;
+				})
+			),
+			NAME_None,
+			bIsBitflagsEnum ? EUserInterfaceActionType::Check : EUserInterfaceActionType::None);
 	}
-}
 
-void SEnumComboBox::OnComboMenuOpening()
-{
-	int32 CurrentNameIndex = Enum->GetIndexByValue(CurrentValue.Get());
-	TSharedPtr<int32> FoundNameIndexItem;
-	for (int32 i = 0; i < VisibleEnumNameIndices.Num(); i++)
-	{
-		if (*VisibleEnumNameIndices[i] == CurrentNameIndex)
-		{
-			FoundNameIndexItem = VisibleEnumNameIndices[i];
-			break;
-		}
-	}
-	if (FoundNameIndexItem.IsValid())
-	{
-		bUpdatingSelectionInternally = true;
-		SetSelectedItem(FoundNameIndexItem);
-		bUpdatingSelectionInternally = false;
-	}
+	return MenuBuilder.MakeWidget();
 }
diff --git a/Engine/Source/Editor/EditorWidgets/Public/SEnumCombo.h b/Engine/Source/Editor/EditorWidgets/Public/SEnumCombo.h
--- a/Engine/Source/Editor/EditorWidgets/Public/SEnumCombo.h
+++ b/Engine/Source/Editor/EditorWidgets/Public/SEnumCombo.h
@@ -3,9 +3,9 @@
 #pragma once
 
 #include "CoreMinimal.h"
-#include "Widgets/Input/SComboBox.h"
+#include "Widgets/Input/SComboButton.h"
 
-class EDITORWIDGETS_API SEnumComboBox : public SComboBox<TSharedPtr<int32>>
+class EDITORWIDGETS_API SEnumComboBox : public SComboButton
 {
 public:
 	DECLARE_DELEGATE_TwoParams(FOnEnumSelectionChanged, int32, ESelectInfo::Type);
@@ -31,22 +31,35 @@ public:
 private:
 	FText GetCurrentValueText() const;
 	FText GetCurrentValueTooltip() const;
-	bool GetValueIsEnabled(TWeakPtr<int32> InValueWeak) const;
-	TSharedRef<SWidget> OnGenerateWidget(TSharedPtr<int32> InItem);
-	void OnComboSelectionChanged(TSharedPtr<int32> InSelectedItem, ESelectInfo::Type SelectInfo);
-	void OnComboMenuOpening();
+	TSharedRef<SWidget> OnGetMenuContent();
 
 private:
+
+	struct FEnumInfo
+	{
+		FEnumInfo() = default;
+		FEnumInfo(const int32 InIndex, const int32 InValue, const FText InDisplayName, const FText InTooltipText)
+			: Index(InIndex), Value(InValue), DisplayName(InDisplayName), TooltipText(InTooltipText)
+		{}
+		
+		int32 Index = 0;
+		int32 Value = 0;
+		FText DisplayName;
+		FText TooltipText;
+	};
+	
 	const UEnum* Enum;
 
 	TAttribute<int32> CurrentValue;
 
 	TAttribute<FSlateFontInfo> Font;
 
-	TArray<TSharedPtr<int32>> VisibleEnumNameIndices;
+	TArray<FEnumInfo> VisibleEnums;
 
 	bool bUpdatingSelectionInternally;
 
+	bool bIsBitflagsEnum;
+	
 	FOnGetToolTipForValue OnGetToolTipForValue;
 
 	FOnEnumSelectionChanged OnEnumSelectionChangedDelegate;
