Reskin asset drop target indicators

diff --git a/Engine/Plugins/FX/Niagara/Source/NiagaraEditor/Private/Widgets/SNiagaraGraphParameterMapGetNode.cpp b/Engine/Plugins/FX/Niagara/Source/NiagaraEditor/Private/Widgets/SNiagaraGraphParameterMapGetNode.cpp
--- a/Engine/Plugins/FX/Niagara/Source/NiagaraEditor/Private/Widgets/SNiagaraGraphParameterMapGetNode.cpp
+++ b/Engine/Plugins/FX/Niagara/Source/NiagaraEditor/Private/Widgets/SNiagaraGraphParameterMapGetNode.cpp
@@ -61,10 +61,6 @@ TSharedRef<SWidget> SNiagaraGraphParameterMapGetNode::CreateNodeContentArea()
 	return 	SNew(SDropTarget)
 		.OnDrop(this, &SNiagaraGraphParameterMapGetNode::OnDroppedOnTarget)
 		.OnAllowDrop(this, &SNiagaraGraphParameterMapGetNode::OnAllowDrop)
-		.HorizontalImage(FNiagaraEditorStyle::Get().GetBrush("NiagaraEditor.DropTarget.BorderHorizontal"))
-		.VerticalImage(FNiagaraEditorStyle::Get().GetBrush("NiagaraEditor.DropTarget.BorderVertical"))
-		.BackgroundColor(FNiagaraEditorStyle::Get().GetColor("NiagaraEditor.DropTarget.BackgroundColor"))
-		.BackgroundColorHover(FNiagaraEditorStyle::Get().GetColor("NiagaraEditor.DropTarget.BackgroundColorHover"))
 		.Content()
 		[
 			SNew(SBorder)
diff --git a/Engine/Plugins/FX/Niagara/Source/NiagaraEditor/Private/Widgets/SNiagaraGraphParameterMapSetNode.cpp b/Engine/Plugins/FX/Niagara/Source/NiagaraEditor/Private/Widgets/SNiagaraGraphParameterMapSetNode.cpp
--- a/Engine/Plugins/FX/Niagara/Source/NiagaraEditor/Private/Widgets/SNiagaraGraphParameterMapSetNode.cpp
+++ b/Engine/Plugins/FX/Niagara/Source/NiagaraEditor/Private/Widgets/SNiagaraGraphParameterMapSetNode.cpp
@@ -32,10 +32,6 @@ TSharedRef<SWidget> SNiagaraGraphParameterMapSetNode::CreateNodeContentArea()
 	return 	SNew(SDropTarget)
 		.OnDrop(this, &SNiagaraGraphParameterMapSetNode::OnDroppedOnTarget)
 		.OnAllowDrop(this, &SNiagaraGraphParameterMapSetNode::OnAllowDrop)
-		.HorizontalImage(FNiagaraEditorStyle::Get().GetBrush("NiagaraEditor.DropTarget.BorderHorizontal"))
-		.VerticalImage(FNiagaraEditorStyle::Get().GetBrush("NiagaraEditor.DropTarget.BorderVertical"))
-		.BackgroundColor(FNiagaraEditorStyle::Get().GetColor("NiagaraEditor.DropTarget.BackgroundColor"))
-		.BackgroundColorHover(FNiagaraEditorStyle::Get().GetColor("NiagaraEditor.DropTarget.BackgroundColorHover"))
 		.Content()
 		[
 			SNew(SBorder)
diff --git a/Engine/Plugins/FX/Niagara/Source/NiagaraEditor/Private/Widgets/SNiagaraParameterPanel.cpp b/Engine/Plugins/FX/Niagara/Source/NiagaraEditor/Private/Widgets/SNiagaraParameterPanel.cpp
--- a/Engine/Plugins/FX/Niagara/Source/NiagaraEditor/Private/Widgets/SNiagaraParameterPanel.cpp
+++ b/Engine/Plugins/FX/Niagara/Source/NiagaraEditor/Private/Widgets/SNiagaraParameterPanel.cpp
@@ -147,10 +147,6 @@ void SNiagaraParameterPanel::Construct(const FArguments& InArgs, const TSharedPt
 			SNew(SDropTarget)
 			.OnDrop(this, &SNiagaraParameterPanel::HandleDragDropOperation)
 			.OnAllowDrop(this, &SNiagaraParameterPanel::GetCanHandleDragDropOperation)
-			.HorizontalImage(FNiagaraEditorStyle::Get().GetBrush("NiagaraEditor.DropTarget.BorderHorizontal"))
-			.VerticalImage(FNiagaraEditorStyle::Get().GetBrush("NiagaraEditor.DropTarget.BorderVertical"))
-			.BackgroundColor(FNiagaraEditorStyle::Get().GetColor("NiagaraEditor.DropTarget.BackgroundColor"))
-			.BackgroundColorHover(FNiagaraEditorStyle::Get().GetColor("NiagaraEditor.DropTarget.BackgroundColorHover"))
 			.Content()
 			[
 				ItemSelector.ToSharedRef()
diff --git a/Engine/Plugins/FX/Niagara/Source/NiagaraEditorWidgets/Private/Stack/SNiagaraStackFunctionInputValue.cpp b/Engine/Plugins/FX/Niagara/Source/NiagaraEditorWidgets/Private/Stack/SNiagaraStackFunctionInputValue.cpp
--- a/Engine/Plugins/FX/Niagara/Source/NiagaraEditorWidgets/Private/Stack/SNiagaraStackFunctionInputValue.cpp
+++ b/Engine/Plugins/FX/Niagara/Source/NiagaraEditorWidgets/Private/Stack/SNiagaraStackFunctionInputValue.cpp
@@ -62,8 +62,6 @@ void SNiagaraStackFunctionInputValue::Construct(const FArguments& InArgs, UNiaga
 		.OnDrop(this, &SNiagaraStackFunctionInputValue::OnFunctionInputDrop)
 		.HorizontalImage(FNiagaraEditorWidgetsStyle::Get().GetBrush("NiagaraEditor.Stack.DropTarget.BorderHorizontal"))
 		.VerticalImage(FNiagaraEditorWidgetsStyle::Get().GetBrush("NiagaraEditor.Stack.DropTarget.BorderVertical"))
-		.BackgroundColor(FNiagaraEditorWidgetsStyle::Get().GetColor("NiagaraEditor.Stack.DropTarget.BackgroundColor"))
-		.BackgroundColorHover(FNiagaraEditorWidgetsStyle::Get().GetColor("NiagaraEditor.Stack.DropTarget.BackgroundColorHover"))
 		.IsEnabled_UObject(FunctionInput, &UNiagaraStackEntry::GetOwnerIsEnabled)
 		.Content()
 		[
diff --git a/Engine/Plugins/FX/Niagara/Source/NiagaraEditorWidgets/Private/Stack/SNiagaraStackModuleItem.cpp b/Engine/Plugins/FX/Niagara/Source/NiagaraEditorWidgets/Private/Stack/SNiagaraStackModuleItem.cpp
--- a/Engine/Plugins/FX/Niagara/Source/NiagaraEditorWidgets/Private/Stack/SNiagaraStackModuleItem.cpp
+++ b/Engine/Plugins/FX/Niagara/Source/NiagaraEditorWidgets/Private/Stack/SNiagaraStackModuleItem.cpp
@@ -347,8 +347,6 @@ TSharedRef<SWidget> SNiagaraStackModuleItem::AddContainerForRowWidgets(TSharedRe
 	.OnDrop(this, &SNiagaraStackModuleItem::OnModuleItemDrop)
 	.HorizontalImage(FNiagaraEditorWidgetsStyle::Get().GetBrush("NiagaraEditor.Stack.DropTarget.BorderHorizontal"))
 	.VerticalImage(FNiagaraEditorWidgetsStyle::Get().GetBrush("NiagaraEditor.Stack.DropTarget.BorderVertical"))
-	.BackgroundColor(FNiagaraEditorWidgetsStyle::Get().GetColor("NiagaraEditor.Stack.DropTarget.BackgroundColor"))
-	.BackgroundColorHover(FNiagaraEditorWidgetsStyle::Get().GetColor("NiagaraEditor.Stack.DropTarget.BackgroundColorHover"))
 	.Content()
 	[
 		SNew(SVerticalBox)
diff --git a/Engine/Source/Editor/EditorStyle/Private/StarshipStyle.cpp b/Engine/Source/Editor/EditorStyle/Private/StarshipStyle.cpp
--- a/Engine/Source/Editor/EditorStyle/Private/StarshipStyle.cpp
+++ b/Engine/Source/Editor/EditorStyle/Private/StarshipStyle.cpp
@@ -361,10 +361,13 @@ void FStarshipEditorStyle::FStyle::SetupGeneralStyles()
 	}
 
 
-	Set( "PlainBorder", new BORDER_BRUSH( "Common/PlainBorder", 2.f/8.f) );
 
-	Set( "WideDash.Horizontal", new IMAGE_BRUSH("Common/WideDash_Horizontal", FVector2D(22, 4), FLinearColor::White, ESlateBrushTileType::Horizontal));
-	Set( "WideDash.Vertical", new IMAGE_BRUSH("Common/WideDash_Vertical", FVector2D(4, 22), FLinearColor::White, ESlateBrushTileType::Vertical));
+	Set("PlainBorder", new BORDER_BRUSH("Common/PlainBorder", 2.f / 8.f));
+
+	Set( "WideDash.Horizontal", new CORE_IMAGE_BRUSH("Starship/Common/Dash_Horizontal", FVector2D(10, 1), FLinearColor::White, ESlateBrushTileType::Horizontal));
+	Set( "WideDash.Vertical", new CORE_IMAGE_BRUSH("Starship/Common/Dash_Vertical", FVector2D(1, 10), FLinearColor::White, ESlateBrushTileType::Vertical));
+
+	Set("DropTarget.Background", new CORE_BOX_BRUSH("Starship/Common/DropTargetBackground", FMargin(6.0f / 64.0f)));
 
 	Set("ThinLine.Horizontal", new IMAGE_BRUSH("Common/ThinLine_Horizontal", FVector2D(11, 2), FLinearColor::White, ESlateBrushTileType::Horizontal));
 
diff --git a/Engine/Source/Editor/EditorWidgets/Private/SDropTarget.cpp b/Engine/Source/Editor/EditorWidgets/Private/SDropTarget.cpp
--- a/Engine/Source/Editor/EditorWidgets/Private/SDropTarget.cpp
+++ b/Engine/Source/Editor/EditorWidgets/Private/SDropTarget.cpp
@@ -23,16 +23,13 @@ void SDropTarget::Construct(const FArguments& InArgs)
 	ValidColor = InArgs._ValidColor;
 	InvalidColor = InArgs._InvalidColor;
 
-	BackgroundColor = InArgs._BackgroundColor;
-	BackgroundColorHover = InArgs._BackgroundColorHover;
-
 	VerticalImage = InArgs._VerticalImage;
 	HorizontalImage = InArgs._HorizontalImage;
 
 	ChildSlot
 	[
 		SNew(SOverlay)
-			
+		.Clipping(EWidgetClipping::ClipToBounds)
 		+ SOverlay::Slot()
 		[
 			InArgs._Content.Widget
@@ -40,20 +37,17 @@ void SDropTarget::Construct(const FArguments& InArgs)
 
 		+ SOverlay::Slot()
 		[
-			SNew(SBox)
+			SNew(SBorder)
 			.Visibility(this, &SDropTarget::GetDragOverlayVisibility)
-			[
-				SNew(SBorder)
-				.BorderImage(InArgs._BackgroundImage)
-				.BorderBackgroundColor(this, &SDropTarget::GetBackgroundBrightness)
-			]
+			.BorderImage(InArgs._BackgroundImage)
+			.BorderBackgroundColor(this, &SDropTarget::GetBackgroundBrightness)
 		]
 	];
 }
 
 FSlateColor SDropTarget::GetBackgroundBrightness() const
 {
-	return ( bIsDragOver ) ? BackgroundColorHover : BackgroundColor;
+	return bAllowDrop ? ValidColor : InvalidColor;
 }
 
 EVisibility SDropTarget::GetDragOverlayVisibility() const
@@ -151,45 +145,47 @@ int32 SDropTarget::OnPaint(const FPaintArgs& Args, const FGeometry& AllottedGeom
 	{
 		if ( bIsDragEventRecognized )
 		{
-			FLinearColor DashColor = bAllowDrop ? ValidColor : InvalidColor;
+			FSlateColor DashColor = bAllowDrop ? ValidColor : InvalidColor;
 
 			int32 DashLayer = LayerId + 1;
+			
+			const float Inset = 3.0f;
 
 			// Top
 			FSlateDrawElement::MakeBox(
 				OutDrawElements,
 				DashLayer,
-				AllottedGeometry.ToPaintGeometry(FVector2D(0, 0), FVector2D(AllottedGeometry.GetLocalSize().X, HorizontalImage->ImageSize.Y)),
+				AllottedGeometry.ToPaintGeometry(FVector2D(Inset, 0), FVector2D(AllottedGeometry.GetLocalSize().X-Inset*2, HorizontalImage->ImageSize.Y)),
 				HorizontalImage,
 				ESlateDrawEffect::None,
-				DashColor);
+				DashColor.GetColor(InWidgetStyle));
 
 			// Bottom
 			FSlateDrawElement::MakeBox(
 				OutDrawElements,
 				DashLayer,
-				AllottedGeometry.ToPaintGeometry(FVector2D(0, AllottedGeometry.GetLocalSize().Y - HorizontalImage->ImageSize.Y), FVector2D(AllottedGeometry.Size.X, HorizontalImage->ImageSize.Y)),
+				AllottedGeometry.ToPaintGeometry(FVector2D(Inset, AllottedGeometry.GetLocalSize().Y - HorizontalImage->ImageSize.Y), FVector2D(AllottedGeometry.Size.X-Inset * 2, HorizontalImage->ImageSize.Y)),
 				HorizontalImage,
 				ESlateDrawEffect::None,
-				DashColor);
+				DashColor.GetColor(InWidgetStyle));
 
 			// Left
 			FSlateDrawElement::MakeBox(
 				OutDrawElements,
 				DashLayer,
-				AllottedGeometry.ToPaintGeometry(FVector2D(0, 0), FVector2D(VerticalImage->ImageSize.X, AllottedGeometry.GetLocalSize().Y)),
+				AllottedGeometry.ToPaintGeometry(FVector2D(0, Inset), FVector2D(VerticalImage->ImageSize.X, AllottedGeometry.GetLocalSize().Y-Inset * 2)),
 				VerticalImage,
 				ESlateDrawEffect::None,
-				DashColor);
+				DashColor.GetColor(InWidgetStyle));
 
 			// Right
 			FSlateDrawElement::MakeBox(
 				OutDrawElements,
 				DashLayer,
-				AllottedGeometry.ToPaintGeometry(FVector2D(AllottedGeometry.GetLocalSize().X - VerticalImage->ImageSize.X, 0), FVector2D(VerticalImage->ImageSize.X, AllottedGeometry.GetLocalSize().Y)),
+				AllottedGeometry.ToPaintGeometry(FVector2D(AllottedGeometry.GetLocalSize().X - VerticalImage->ImageSize.X, Inset), FVector2D(VerticalImage->ImageSize.X, AllottedGeometry.GetLocalSize().Y-Inset * 2)),
 				VerticalImage,
 				ESlateDrawEffect::None,
-				DashColor);
+				DashColor.GetColor(InWidgetStyle));
 
 			return DashLayer;
 		}
diff --git a/Engine/Source/Editor/EditorWidgets/Public/SDropTarget.h b/Engine/Source/Editor/EditorWidgets/Public/SDropTarget.h
--- a/Engine/Source/Editor/EditorWidgets/Public/SDropTarget.h
+++ b/Engine/Source/Editor/EditorWidgets/Public/SDropTarget.h
@@ -9,6 +9,7 @@
 #include "Widgets/DeclarativeSyntaxSupport.h"
 #include "Widgets/SCompoundWidget.h"
 #include "EditorStyleSet.h"
+#include "Styling/StyleColors.h"
 
 class FPaintArgs;
 class FSlateWindowElementList;
@@ -25,30 +26,36 @@ public:
 	DECLARE_DELEGATE_RetVal_OneParam(bool, FVerifyDrag, TSharedPtr<FDragDropOperation>);
 
 	SLATE_BEGIN_ARGS(SDropTarget)
-		: _ValidColor(FLinearColor(0, 1, 0, 1))
-		, _InvalidColor(FLinearColor(1, 0, 0, 1))
-		, _BackgroundColor(FLinearColor(1, 1, 1, 0.50f))
-		, _BackgroundColorHover(FLinearColor(1, 1, 1, 0.25f))
+		: _ValidColor(FStyleColors::AccentBlue)
+		, _InvalidColor(FStyleColors::Error)
 		, _VerticalImage(FEditorStyle::GetBrush("WideDash.Vertical"))
 		, _HorizontalImage(FEditorStyle::GetBrush("WideDash.Horizontal"))
-		, _BackgroundImage(FEditorStyle::GetBrush("WhiteBrush"))
+		, _BackgroundImage(FEditorStyle::GetBrush("DropTarget.Background"))
 	{ }
+	
+		UE_DEPRECATED(5.0, "BackgroundColor has been removed. You may alter the background brush to get the same effect.")
+		FArguments& BackgroundColor(const FLinearColor& InBackgroundColor)
+		{
+			return Me();
+		}
+
+		UE_DEPRECATED(5.0, "BackgroundColorHover has been removed. You may alter the background brush when hovered to get the same effect.")
+		FArguments& BackgroundColorHover(const FLinearColor& InBackgroundColor)
+		{
+			return Me();
+		}
 		/* Content to display for the in the drop target */
 		SLATE_DEFAULT_SLOT( FArguments, Content )
 		/** The color of the vertical/horizontal images when the drop data is valid */
-		SLATE_ARGUMENT(FLinearColor, ValidColor)
+		SLATE_ARGUMENT(FSlateColor, ValidColor)
 		/** The color of the vertical/horizontal images when the drop data is not valid */
-		SLATE_ARGUMENT(FLinearColor, InvalidColor)
-		/** The color multiplier that is applied to the background image. */
-		SLATE_ARGUMENT(FLinearColor, BackgroundColor)
-		/** The color multiplier that is applied to the background image on hover. */
-		SLATE_ARGUMENT(FLinearColor, BackgroundColorHover)
+		SLATE_ARGUMENT(FSlateColor, InvalidColor)
 		/** Vertical border image that is used. */
 		SLATE_ARGUMENT(const FSlateBrush*, VerticalImage)
 		/** Horizontal border image that is used. */
 		SLATE_ARGUMENT(const FSlateBrush*, HorizontalImage)
 		/** The background image that is applied after the surface. */
-		SLATE_ARGUMENT(const FSlateBrush*, BackgroundImage)
+		SLATE_ATTRIBUTE(const FSlateBrush*, BackgroundImage)
 		/** Called when a valid asset is dropped */
 		SLATE_EVENT(FOnDrop, OnDrop)
 		/** Called to check if an asset is acceptable for dropping */
@@ -93,15 +100,9 @@ private:
 	FVerifyDrag IsRecognizedEvent;
 
 	/** The color of the vertical/horizontal images when the drop data is valid */
-	FLinearColor ValidColor;
+	FSlateColor ValidColor;
 	/** The color of the vertical/horizontal images when the drop data is not valid */
-	FLinearColor InvalidColor;
-	/** The color multiplier that is applied to the background image. */
-	FLinearColor BackgroundColor;
-	/** The color multiplier that is applied to the background image on hover. */
-	FLinearColor BackgroundColorHover;
-	/** The background image that is applied after the surface. */
-	const FSlateBrush* BackgroundImage;
+	FSlateColor InvalidColor;
 	/** Vertical border image that is used. */
 	const FSlateBrush* VerticalImage;
 	/** Horizontal border image that is used. */
diff --git a/Engine/Source/Editor/Sequencer/Private/SSequencerTrackArea.cpp b/Engine/Source/Editor/Sequencer/Private/SSequencerTrackArea.cpp
--- a/Engine/Source/Editor/Sequencer/Private/SSequencerTrackArea.cpp
+++ b/Engine/Source/Editor/Sequencer/Private/SSequencerTrackArea.cpp
@@ -17,6 +17,7 @@
 #include "DisplayNodes/SequencerTrackNode.h"
 #include "DisplayNodes/SequencerObjectBindingNode.h"
 #include "CommonMovieSceneTools.h"
+#include "Styling/StyleColors.h"
 
 FTrackAreaSlot::FTrackAreaSlot(const TSharedPtr<SSequencerTrackLane>& InSlotContent)
 	: TAlignmentWidgetSlotMixin<FTrackAreaSlot>(HAlign_Fill, VAlign_Top)
@@ -202,7 +203,7 @@ int32 SSequencerTrackArea::OnPaint(const FPaintArgs& Args, const FGeometry& Allo
 			
 			FGeometry NodeGeometry = TrackLane.Get()->GetCachedGeometry();
 
-			FLinearColor DashColor = bAllowDrop ? FLinearColor::Green: FLinearColor::Red;
+			FSlateColor DashColor = bAllowDrop ? FStyleColors::AccentBlue : FStyleColors::Error;
 
 			const FSlateBrush* HorizontalBrush = FEditorStyle::GetBrush("WideDash.Horizontal");
 			const FSlateBrush* VerticalBrush = FEditorStyle::GetBrush("WideDash.Vertical");
@@ -226,7 +227,7 @@ int32 SSequencerTrackArea::OnPaint(const FPaintArgs& Args, const FGeometry& Allo
 				AllottedGeometry.ToPaintGeometry(FVector2D(DropMinX, TrackLane.Get()->GetPhysicalPosition()), FVector2D(DropMaxX-DropMinX, HorizontalBrush->ImageSize.Y)),
 				HorizontalBrush,
 				ESlateDrawEffect::None,
-				DashColor);
+				DashColor.GetSpecifiedColor());
 
 			// Bottom
 			FSlateDrawElement::MakeBox(
@@ -235,7 +236,7 @@ int32 SSequencerTrackArea::OnPaint(const FPaintArgs& Args, const FGeometry& Allo
 				AllottedGeometry.ToPaintGeometry(FVector2D(DropMinX, TrackLane.Get()->GetPhysicalPosition() + (TrackLane.Get()->GetCachedGeometry().GetLocalSize().Y - HorizontalBrush->ImageSize.Y)), FVector2D(DropMaxX-DropMinX, HorizontalBrush->ImageSize.Y)),
 				HorizontalBrush,
 				ESlateDrawEffect::None,
-				DashColor);
+				DashColor.GetSpecifiedColor());
 
 			// Left
 			FSlateDrawElement::MakeBox(
@@ -244,7 +245,7 @@ int32 SSequencerTrackArea::OnPaint(const FPaintArgs& Args, const FGeometry& Allo
 				AllottedGeometry.ToPaintGeometry(FVector2D(DropMinX, TrackLane.Get()->GetPhysicalPosition()), FVector2D(VerticalBrush->ImageSize.X, TrackLane.Get()->GetCachedGeometry().GetLocalSize().Y)),
 				VerticalBrush,
 				ESlateDrawEffect::None,
-				DashColor);
+				DashColor.GetSpecifiedColor());
 
 			// Right
 			FSlateDrawElement::MakeBox(
@@ -253,7 +254,7 @@ int32 SSequencerTrackArea::OnPaint(const FPaintArgs& Args, const FGeometry& Allo
 				AllottedGeometry.ToPaintGeometry(FVector2D(DropMaxX - VerticalBrush->ImageSize.X, TrackLane.Get()->GetPhysicalPosition()), FVector2D(VerticalBrush->ImageSize.X, TrackLane.Get()->GetCachedGeometry().GetLocalSize().Y)),
 				VerticalBrush,
 				ESlateDrawEffect::None,
-				DashColor);
+				DashColor.GetSpecifiedColor());
 		}
 	}
 
