Added a new interface for displaying names/tooltips in the details window and added a new name display/tooltip for instanced static mesh proxy objects.

diff --git a/Engine/Source/Editor/EditorWidgets/Private/EditorWidgetsModule.cpp b/Engine/Source/Editor/EditorWidgets/Private/EditorWidgetsModule.cpp
--- a/Engine/Source/Editor/EditorWidgets/Private/EditorWidgetsModule.cpp
+++ b/Engine/Source/Editor/EditorWidgets/Private/EditorWidgetsModule.cpp
@@ -4,6 +4,7 @@
 #include "Modules/ModuleManager.h"
 #include "Widgets/DeclarativeSyntaxSupport.h"
 #include "SObjectNameEditableTextBox.h"
+#include "ObjectNameEditSinkRegistry.h"
 #include "SAssetDiscoveryIndicator.h"
 #include "ITransportControl.h"
 #include "STransportControl.h"
@@ -14,15 +15,17 @@ const FName FEditorWidgetsModule::EditorWidgetsAppIdentifier( TEXT( "EditorWidge
 
 void FEditorWidgetsModule::StartupModule()
 {
+	ObjectNameEditSinkRegistry = MakeShared<UE::EditorWidgets::FObjectNameEditSinkRegistry>();
 }
 
 void FEditorWidgetsModule::ShutdownModule()
 {
+	ObjectNameEditSinkRegistry.Reset();
 }
 
 TSharedRef<IObjectNameEditableTextBox> FEditorWidgetsModule::CreateObjectNameEditableTextBox(const TArray<TWeakObjectPtr<UObject>>& Objects)
 {
-	TSharedRef<SObjectNameEditableTextBox> Widget = SNew(SObjectNameEditableTextBox).Objects(Objects);
+	TSharedRef<SObjectNameEditableTextBox> Widget = SNew(SObjectNameEditableTextBox).Objects(Objects).Registry(ObjectNameEditSinkRegistry);
 	return Widget;
 }
 
@@ -39,3 +42,8 @@ TSharedRef<ITransportControl> FEditorWidgetsModule::CreateTransportControl(const
 	return SNew(STransportControl)
 		.TransportArgs(Args);
 }
+
+TSharedRef<UE::EditorWidgets::FObjectNameEditSinkRegistry> FEditorWidgetsModule::GetObjectNameEditSinkRegistry() const
+{
+	return ObjectNameEditSinkRegistry.ToSharedRef();
+}
\ No newline at end of file
diff --git a/Engine/Source/Editor/EditorWidgets/Private/ObjectNameEditSinkRegistry.cpp b/Engine/Source/Editor/EditorWidgets/Private/ObjectNameEditSinkRegistry.cpp
new file mode 100644
--- /dev/null
+++ b/Engine/Source/Editor/EditorWidgets/Private/ObjectNameEditSinkRegistry.cpp
@@ -0,0 +1,117 @@
+// Copyright Epic Games, Inc. All Rights Reserved.
+
+#include "ObjectNameEditSinkRegistry.h"
+#include "Templates/SharedPointer.h"
+#include "GameFramework/Actor.h"
+
+#define LOCTEXT_NAMESPACE "EditorWidgets"
+
+namespace UE::EditorWidgets
+{
+
+/** Name edit sink implementation for a generic UObject. Used for anything without a specific implementation. */
+class FObjectNameEditSink : public IObjectNameEditSink
+{
+	virtual UClass* GetSupportedClass() const override
+	{
+		return UObject::StaticClass();
+	}
+
+	virtual FText GetObjectDisplayName(UObject* Object) const override
+	{
+		return FText::FromString(Object->GetName());
+	}
+
+	virtual FText GetObjectNameTooltip(UObject* Object) const override
+	{
+		return LOCTEXT("EditableActorLabel_NoEditObjectTooltip", "Can't rename selected object (only actors can have editable labels)");
+	}
+};
+
+class FActorNameEditSink : public IObjectNameEditSink
+{
+	virtual UClass* GetSupportedClass() const override
+	{
+		return AActor::StaticClass();
+	}
+
+	virtual FText GetObjectDisplayName(UObject* Object) const override
+	{
+		return FText::FromString(CastChecked<AActor>(Object)->GetActorLabel());
+	}
+
+	virtual bool IsObjectDisplayNameReadOnly(UObject* Object) const override
+	{
+		return !(CastChecked<AActor>(Object)->IsActorLabelEditable());
+	};
+
+	virtual bool SetObjectDisplayName(UObject* Object, FString DisplayName) override
+	{
+		if (IsObjectDisplayNameReadOnly(Object))
+		{
+			return false;
+		}
+
+		AActor* Actor = CastChecked<AActor>(Object);
+
+		if (Actor->GetActorLabel() == DisplayName)
+		{
+			return false;
+		}
+
+		FActorLabelUtilities::RenameExistingActor(Actor, DisplayName);
+		return true;
+	}
+
+	virtual FText GetObjectNameTooltip(UObject* Object) const override
+	{
+		if (IsObjectDisplayNameReadOnly(Object))
+		{
+			return LOCTEXT("EditableActorLabel_NoEditActorTooltip", "Can't rename selected actor (its label isn't editable)");
+		}
+
+		return FText::Format(LOCTEXT("EditableActorLabel_ActorTooltipFmt", "Rename the selected {0}"), FText::FromString(Object->GetClass()->GetName()));
+	}
+};
+
+FObjectNameEditSinkRegistry::FObjectNameEditSinkRegistry()
+{
+	// Register default name edit sinks
+	RegisterObjectNameEditSink(MakeShared<FObjectNameEditSink>());
+	RegisterObjectNameEditSink(MakeShared<FActorNameEditSink>());
+}
+
+void FObjectNameEditSinkRegistry::RegisterObjectNameEditSink(const TSharedRef<IObjectNameEditSink>& NewSink)
+{
+	ObjectNameEditSinkList.Add(NewSink);
+}
+
+void FObjectNameEditSinkRegistry::UnregisterObjectNameEditSink(const TSharedRef<IObjectNameEditSink>& SinkToRemove)
+{
+	ObjectNameEditSinkList.Remove(SinkToRemove);
+}
+
+TSharedPtr<IObjectNameEditSink> FObjectNameEditSinkRegistry::GetObjectNameEditSinkForClass(const UClass* Class) const
+{
+	TSharedPtr<IObjectNameEditSink> MostDerivedObjectNameEditSink;
+
+	for (const TSharedRef<IObjectNameEditSink>& NameEditSink : ObjectNameEditSinkList)
+	{
+		UClass* SupportedClass = NameEditSink->GetSupportedClass();
+
+		if (Class->IsChildOf(SupportedClass))
+		{
+			if (!MostDerivedObjectNameEditSink.IsValid() ||
+				SupportedClass->IsChildOf(MostDerivedObjectNameEditSink->GetSupportedClass()))
+			{
+				MostDerivedObjectNameEditSink = NameEditSink;
+			}
+		}
+	}
+
+	return MostDerivedObjectNameEditSink;
+}
+
+} // end namespace UE::EditorWidgets
+
+#undef LOCTEXT_NAMESPACE
diff --git a/Engine/Source/Editor/EditorWidgets/Private/SObjectNameEditableTextBox.cpp b/Engine/Source/Editor/EditorWidgets/Private/SObjectNameEditableTextBox.cpp
--- a/Engine/Source/Editor/EditorWidgets/Private/SObjectNameEditableTextBox.cpp
+++ b/Engine/Source/Editor/EditorWidgets/Private/SObjectNameEditableTextBox.cpp
@@ -1,6 +1,7 @@
 // Copyright Epic Games, Inc. All Rights Reserved.
 
 #include "SObjectNameEditableTextBox.h"
+#include "ObjectNameEditSinkRegistry.h"
 #include "Rendering/DrawElements.h"
 #include "Framework/Application/SlateApplication.h"
 #include "Widgets/Text/SInlineEditableTextBlock.h"
@@ -25,6 +26,7 @@ void SObjectNameEditableTextBox::Construct( const FArguments& InArgs )
 	bUpdateHighlightSpring = false;
 
 	Objects = InArgs._Objects;
+	ObjectNameEditSinkRegistry = InArgs._Registry;
 
 	ChildSlot
 	[
@@ -129,7 +131,7 @@ int32 SObjectNameEditableTextBox::OnPaint( const FPaintArgs& Args, const FGeomet
 
 FText SObjectNameEditableTextBox::GetNameText() const
 {
-	FString Result;
+	FText Result;
 
 	if (Objects.Num() == 1)
 	{
@@ -139,34 +141,39 @@ FText SObjectNameEditableTextBox::GetNameText() const
 	{
 		if (!UserSetCommonName.IsEmpty())
 		{
-			Result = UserSetCommonName;
+			Result = FText::FromString(UserSetCommonName);
 		}
 	}
 
-	return FText::FromString(Result);
+	return Result;
 }
 
 FText SObjectNameEditableTextBox::GetNameTooltipText() const
 {
-	FText Result = FText::GetEmpty();
+	FText Result;
 
 	if (Objects.Num() == 0)
 	{
 		Result = LOCTEXT("EditableActorLabel_NoObjectsTooltip", "Nothing selected");
 	}
-	else if (Objects.Num() == 1 && Objects[0].IsValid())
+	else if (Objects.Num() == 1)
 	{
-		if (!IsReadOnly())
+		UObject* ObjectPtr = Objects[0].Get();
+		if (!ObjectPtr)
 		{
-			Result = FText::Format(LOCTEXT("EditableActorLabel_ActorTooltipFmt", "Rename the selected {0}"), FText::FromString(Objects[0].Get()->GetClass()->GetName()));
+			return Result;
 		}
-		else if (Objects[0].Get()->IsA(AActor::StaticClass()))
+
+		TSharedPtr<UE::EditorWidgets::FObjectNameEditSinkRegistry> RegistryPtr = ObjectNameEditSinkRegistry.Pin();
+		if (!RegistryPtr.IsValid())
 		{
-			Result = LOCTEXT("EditableActorLabel_NoEditActorTooltip", "Can't rename selected actor (its label isn't editable)");
+			return Result;
 		}
-		else
+
+		TSharedPtr<UE::EditorWidgets::IObjectNameEditSink> ObjectNameEditPtr = RegistryPtr->GetObjectNameEditSinkForClass(ObjectPtr->GetClass());
+		if (ObjectNameEditPtr)
 		{
-			Result = LOCTEXT("EditableActorLabel_NoEditObjectTooltip", "Can't rename selected object (only actors can have editable labels)");
+			Result = ObjectNameEditPtr->GetObjectNameTooltip(ObjectPtr);
 		}
 	}
 	else if (Objects.Num() > 1)
@@ -189,30 +196,6 @@ EVisibility SObjectNameEditableTextBox::GetNameVisibility() const
 	return ((Objects.Num() == 1 && Objects[0].IsValid()) || Objects.Num() > 1)
 		? EVisibility::Visible
 		: EVisibility::Collapsed;
-} 
-
-
-static bool RenameActor(TWeakObjectPtr<UObject> Object, FString Name)
-{
-	// Apply the change to the selected actor
-	if (!Object.IsValid() || !Object.Get()->IsA(AActor::StaticClass()))
-	{
-		return false;
-	}
-
-	AActor* Actor = Cast<AActor>(Object.Get());
-	if (!Actor->IsActorLabelEditable())
-	{
-		return false;
-	}
-
-	if (Actor->GetActorLabel() == Name)
-	{
-		return false;
-	}
-
-	FActorLabelUtilities::RenameExistingActor(Actor, Name);
-	return true;
 }
 
 void SObjectNameEditableTextBox::OnNameTextCommitted(const FText& NewText, ETextCommit::Type InTextCommit)
@@ -224,13 +207,18 @@ void SObjectNameEditableTextBox::OnNameTextCommitted(const FText& NewText, EText
 		return;
 	}
 
-
 	FText TrimmedText = FText::TrimPrecedingAndTrailing(NewText);
 	if (TrimmedText.IsEmpty())
 	{
 		return;
 	}
 
+	TSharedPtr<UE::EditorWidgets::FObjectNameEditSinkRegistry> RegistryPtr = ObjectNameEditSinkRegistry.Pin();
+	if (!RegistryPtr.IsValid())
+	{
+		return;
+	}
+
 	UserSetCommonName = TrimmedText.ToString();
 
 	const FScopedTransaction Transaction( LOCTEXT("RenameActorsTransaction", "Rename Multiple Actors") );
@@ -238,7 +226,14 @@ void SObjectNameEditableTextBox::OnNameTextCommitted(const FText& NewText, EText
 	bool bChanged = false;
 	for (TWeakObjectPtr<UObject> Object : Objects)
 	{
-		if (RenameActor(Object, UserSetCommonName))
+		UObject* ObjectPtr = Object.Get();
+		if (!ObjectPtr)
+		{
+			continue;
+		}
+
+		TSharedPtr<UE::EditorWidgets::IObjectNameEditSink> ObjectNameEditPtr = RegistryPtr->GetObjectNameEditSinkForClass(ObjectPtr->GetClass());
+		if (ObjectNameEditPtr && ObjectNameEditPtr->SetObjectDisplayName(ObjectPtr, UserSetCommonName))
 		{
 			bChanged = true;
 		}
@@ -266,46 +261,50 @@ bool SObjectNameEditableTextBox::IsReadOnly() const
 		return true;
 	}
 
+	TSharedPtr<UE::EditorWidgets::FObjectNameEditSinkRegistry> RegistryPtr = ObjectNameEditSinkRegistry.Pin();
+	// this shouldn't happen but let's be safe and say readonly
+	if (!RegistryPtr.IsValid())
+	{
+		return true;
+	}
+
 	for (TWeakObjectPtr<UObject> Object : Objects)
 	{
-		if (Object.IsValid())
+		UObject* ObjectPtr = Object.Get();
+		if (!ObjectPtr)
 		{
-			if (Object.Get()->IsA(AActor::StaticClass()))
-			{
-				AActor* Actor = (AActor*) Object.Get();
-				if (!Actor->IsActorLabelEditable())
-				{
-					// can't edit the name when a non-editable actor is selected
-					return true;
-				}
-			}
-			else
-			{
-				// can't edit the name when a non-actor is selected
-				return true;
-			}
+			continue;
 		}
-	}
 
+		TSharedPtr<UE::EditorWidgets::IObjectNameEditSink> ObjectNameEditPtr = RegistryPtr->GetObjectNameEditSinkForClass(ObjectPtr->GetClass());
+		if (ObjectNameEditPtr && ObjectNameEditPtr->IsObjectDisplayNameReadOnly(ObjectPtr))
+		{
+			// if any objects are read only then the whole box is
+			return true;
+		}
+	}
 	return false;
 }
 
-FString SObjectNameEditableTextBox::GetObjectDisplayName(TWeakObjectPtr<UObject> Object)
+FText SObjectNameEditableTextBox::GetObjectDisplayName(TWeakObjectPtr<UObject> Object) const
 {
-	FString Result;
-	if(Object.IsValid())
+	UObject* ObjectPtr = Object.Get();
+	// no display name for no object
+	if (!ObjectPtr)
 	{
-		UObject* ObjectPtr = Object.Get();
-		if (ObjectPtr->IsA(AActor::StaticClass()))
-		{
-			Result =  ((AActor*)ObjectPtr)->GetActorLabel();
-		}
-		else
+		return FText();
+	}
+	TSharedPtr<UE::EditorWidgets::FObjectNameEditSinkRegistry> RegistryPtr = ObjectNameEditSinkRegistry.Pin();
+	if (RegistryPtr.IsValid())
+	{
+		TSharedPtr<UE::EditorWidgets::IObjectNameEditSink> ObjectNameEditPtr = RegistryPtr->GetObjectNameEditSinkForClass(ObjectPtr->GetClass());
+		if (ObjectNameEditPtr)
 		{
-			Result = ObjectPtr->GetName();
+			return ObjectNameEditPtr->GetObjectDisplayName(ObjectPtr);
 		}
 	}
-	return Result;
+	// this shouldn't happen but let's be safe and just give the default name
+	return FText::FromString(ObjectPtr->GetName());
 }
 
 // No code beyond this point
diff --git a/Engine/Source/Editor/EditorWidgets/Private/SObjectNameEditableTextBox.h b/Engine/Source/Editor/EditorWidgets/Private/SObjectNameEditableTextBox.h
--- a/Engine/Source/Editor/EditorWidgets/Private/SObjectNameEditableTextBox.h
+++ b/Engine/Source/Editor/EditorWidgets/Private/SObjectNameEditableTextBox.h
@@ -21,6 +21,8 @@ public:
 
 		SLATE_ARGUMENT(TArray<TWeakObjectPtr<UObject>>, Objects)
 
+		SLATE_ARGUMENT(TWeakPtr<UE::EditorWidgets::FObjectNameEditSinkRegistry>, Registry)
+
 	SLATE_END_ARGS()
 
 	/**
@@ -55,8 +57,8 @@ private:
 	/** Getter for the IsReadOnly attribute of the editable text inside this widget */
 	bool IsReadOnly() const;
 
-	/** Helper class the get the object name or the actor label if an object is an actor */
-	static FString GetObjectDisplayName(TWeakObjectPtr<UObject> Object);
+	/** Helper function to access name registry and retrieve display name */
+	FText GetObjectDisplayName(TWeakObjectPtr<UObject> Object) const;
 
 	/** The list of objects whose names are edited by the widget */
 	TArray<TWeakObjectPtr<UObject>> Objects;
@@ -94,6 +96,9 @@ private:
 	/** The text box used to edit object names */ 
 	TSharedPtr< SInlineEditableTextBlock > TextBox;
 
+	/** The registry for object names */
+	TWeakPtr<UE::EditorWidgets::FObjectNameEditSinkRegistry> ObjectNameEditSinkRegistry;
+
 	// Temp flag to trigger a highlight spring update in the passive tick (because that's where the geometry is)
 	bool bUpdateHighlightSpring;
 };
diff --git a/Engine/Source/Editor/EditorWidgets/Public/EditorWidgetsModule.h b/Engine/Source/Editor/EditorWidgets/Public/EditorWidgetsModule.h
--- a/Engine/Source/Editor/EditorWidgets/Public/EditorWidgetsModule.h
+++ b/Engine/Source/Editor/EditorWidgets/Public/EditorWidgetsModule.h
@@ -12,6 +12,11 @@
 class ITransportControl;
 struct FTransportControlArgs;
 
+namespace UE::EditorWidgets
+{
+class FObjectNameEditSinkRegistry;
+}
+
 /** Interface for the widget that wraps an editable text box for viewing the names of objects or editing the labels of actors */
 class IObjectNameEditableTextBox : public SCompoundWidget
 {
@@ -59,6 +64,17 @@ public:
 	 */
 	virtual TSharedRef<ITransportControl> CreateTransportControl(const FTransportControlArgs& Args);
 
+	/**
+	 * Gets the registry of name edit 'sinks' used for getting and editing object display names
+	 * 
+	 * @return The registry
+	 */
+	virtual TSharedRef<UE::EditorWidgets::FObjectNameEditSinkRegistry> GetObjectNameEditSinkRegistry() const;
+
 	/** Editor Widgets app identifier string */
 	static const FName EditorWidgetsAppIdentifier;
+
+private:
+	/** Registry for name interfaces */
+	TSharedPtr<UE::EditorWidgets::FObjectNameEditSinkRegistry> ObjectNameEditSinkRegistry;
 };
diff --git a/Engine/Source/Editor/EditorWidgets/Public/IObjectNameEditSink.h b/Engine/Source/Editor/EditorWidgets/Public/IObjectNameEditSink.h
new file mode 100644
--- /dev/null
+++ b/Engine/Source/Editor/EditorWidgets/Public/IObjectNameEditSink.h
@@ -0,0 +1,42 @@
+// Copyright Epic Games, Inc. All Rights Reserved.
+
+#pragma once
+
+#include "UObject/Object.h"
+
+namespace UE::EditorWidgets
+{
+
+/** ObjectNameEditSink provides an interface for the in-editor object name */
+class IObjectNameEditSink
+{
+public:
+	/** Virtual destructor */
+	virtual ~IObjectNameEditSink() = default;
+
+	/** Checks to see if the specified object is handled by this type. */
+	virtual UClass* GetSupportedClass() const = 0;
+
+	/** Returns the display name for that object. */
+	virtual FText GetObjectDisplayName(UObject* Object) const = 0;
+
+	/** Whether or not the name can be edited */
+	virtual bool IsObjectDisplayNameReadOnly(UObject* Object) const
+	{
+		return true;
+	};
+
+	/** 
+	 * Attempt to rename the object. Will fail if IsObjectDisplayNameReadOnly returns true.
+	 * @return Whether or not the display name was successfully set
+	 */
+	virtual bool SetObjectDisplayName(UObject* Object, FString DisplayName)
+	{
+		return false;
+	};
+
+	/** Returns the tooltip text for that object */
+	virtual FText GetObjectNameTooltip(UObject* Object) const = 0;
+};
+
+} //end namespace UE::EditorWidgets
diff --git a/Engine/Source/Editor/EditorWidgets/Public/ObjectNameEditSinkRegistry.h b/Engine/Source/Editor/EditorWidgets/Public/ObjectNameEditSinkRegistry.h
new file mode 100644
--- /dev/null
+++ b/Engine/Source/Editor/EditorWidgets/Public/ObjectNameEditSinkRegistry.h
@@ -0,0 +1,29 @@
+// Copyright Epic Games, Inc. All Rights Reserved.
+
+#pragma once
+
+#include "IObjectNameEditSink.h"
+
+namespace UE::EditorWidgets
+{
+
+class EDITORWIDGETS_API FObjectNameEditSinkRegistry
+{
+public:
+	FObjectNameEditSinkRegistry();
+
+	/** Registers an object name edit sink so it can provide a name editing interface for an object type. */
+	void RegisterObjectNameEditSink(const TSharedRef<IObjectNameEditSink>& NewSink);
+
+	/** Unregisters a name edit sink. It will no longer provide a name editing interface for an object type. */
+	void UnregisterObjectNameEditSink(const TSharedRef<IObjectNameEditSink>& SinkToRemove);
+
+	/** Gets the appropriate ObjectNameEditSink for the supplied class */
+	TSharedPtr<IObjectNameEditSink> GetObjectNameEditSinkForClass(const UClass* Class) const;
+
+private:
+	/** The list of all registered ObjectNameEditSinks */
+	TArray<TSharedRef<IObjectNameEditSink>> ObjectNameEditSinkList;
+};
+
+} // end namespace UE::EditorWidgets
diff --git a/Engine/Source/Editor/UnrealEd/Private/Elements/Framework/EditorElements.cpp b/Engine/Source/Editor/UnrealEd/Private/Elements/Framework/EditorElements.cpp
--- a/Engine/Source/Editor/UnrealEd/Private/Elements/Framework/EditorElements.cpp
+++ b/Engine/Source/Editor/UnrealEd/Private/Elements/Framework/EditorElements.cpp
@@ -22,6 +22,11 @@
 #include "Elements/SMInstance/SMInstanceElementDetailsInterface.h"
 #include "Elements/SMInstance/SMInstanceElementEditorWorldInterface.h"
 #include "Elements/SMInstance/SMInstanceElementEditorSelectionInterface.h"
+#include "Elements/SMInstance/SMInstanceElementDetailsProxyObject.h"
+
+#include "Modules/ModuleManager.h"
+#include "EditorWidgetsModule.h"
+#include "ObjectNameEditSinkRegistry.h"
 
 FSimpleMulticastDelegate OnRegisterEditorElementsDelegate;
 
@@ -59,6 +64,9 @@ void RegisterEditorSMInstanceElements()
 	Registry->RegisterElementInterface<UTypedElementDetailsInterface>(NAME_SMInstance, NewObject<USMInstanceElementDetailsInterface>());
 	Registry->RegisterElementInterface<UTypedElementWorldInterface>(NAME_SMInstance, NewObject<USMInstanceElementEditorWorldInterface>(), /*bAllowOverride*/true);
 	Registry->RegisterElementInterface<UTypedElementSelectionInterface>(NAME_SMInstance, NewObject<USMInstanceElementEditorSelectionInterface>(), /*bAllowOverride*/true);
+
+	FEditorWidgetsModule& EditorWidgetsModule = FModuleManager::LoadModuleChecked<FEditorWidgetsModule>("EditorWidgets");
+	EditorWidgetsModule.GetObjectNameEditSinkRegistry()->RegisterObjectNameEditSink(MakeShared<FSMInstanceElementDetailsProxyObjectNameEditSink>());
 }
 
 void RegisterEditorElements()
diff --git a/Engine/Source/Editor/UnrealEd/Private/Elements/SMInstance/SMInstanceElementDetailsProxyObject.cpp b/Engine/Source/Editor/UnrealEd/Private/Elements/SMInstance/SMInstanceElementDetailsProxyObject.cpp
--- a/Engine/Source/Editor/UnrealEd/Private/Elements/SMInstance/SMInstanceElementDetailsProxyObject.cpp
+++ b/Engine/Source/Editor/UnrealEd/Private/Elements/SMInstance/SMInstanceElementDetailsProxyObject.cpp
@@ -8,6 +8,8 @@
 #include "UnrealEdGlobals.h"
 #include "Editor/UnrealEdEngine.h"
 
+#define LOCTEXT_NAMESPACE "SMInstanceElementDetails"
+
 void USMInstanceElementDetailsProxyObject::Initialize(const FSMInstanceElementId& InSMInstanceElementId)
 {
 	ISMComponent = InSMInstanceElementId.ISMComponent;
@@ -96,3 +98,46 @@ FSMInstanceManager USMInstanceElementDetailsProxyObject::GetSMInstance() const
 	const FSMInstanceId SMInstanceId = FSMInstanceElementIdMap::Get().GetSMInstanceIdFromSMInstanceElementId(FSMInstanceElementId{ ISMComponent.Get(), ISMInstanceId });
 	return FSMInstanceManager(SMInstanceId, SMInstanceElementDataUtil::GetSMInstanceManager(SMInstanceId));
 }
+
+UClass* FSMInstanceElementDetailsProxyObjectNameEditSink::GetSupportedClass() const
+{
+	return USMInstanceElementDetailsProxyObject::StaticClass();
+}
+
+FText FSMInstanceElementDetailsProxyObjectNameEditSink::GetObjectDisplayName(UObject* Object) const
+{
+	FText DisplayName;
+	USMInstanceElementDetailsProxyObject* ProxyObject = CastChecked<USMInstanceElementDetailsProxyObject>(Object);
+	FSMInstanceManager SMInstance = ProxyObject->GetSMInstance();
+	if (!SMInstance)
+	{
+		return LOCTEXT("ProxyObjectDisplayNameUnknown", "Unknown Instanced Static Mesh");
+	}
+	TObjectPtr<UStaticMesh> StaticMesh = SMInstance.GetISMComponent()->GetStaticMesh();
+	if (StaticMesh)
+	{
+		DisplayName = FText::FromString(StaticMesh->GetName());
+	}
+	else
+	{
+		DisplayName = LOCTEXT("ProxyObjectDisplayNameUnassigned", "Static Mesh Unassigned");
+	}
+	DisplayName = FText::Format(LOCTEXT("ProxyObjectDisplayNameFmt", "{0} - Instance {1}"), DisplayName, SMInstance.GetInstanceId().InstanceIndex);
+	return DisplayName;
+}
+
+FText FSMInstanceElementDetailsProxyObjectNameEditSink::GetObjectNameTooltip(UObject* Object) const
+{
+	USMInstanceElementDetailsProxyObject* ProxyObject = CastChecked<USMInstanceElementDetailsProxyObject>(Object);
+	FSMInstanceManager SMInstance = ProxyObject->GetSMInstance();
+	if (!SMInstance)
+	{
+		return LOCTEXT("ProxyObjectTooltipUnknown", "Unknown Instanced Static Mesh");;
+	}
+	return FText::Format(LOCTEXT("ProxyObjectTooltipFmt", "Instance {0} on {1}"),
+		SMInstance.GetInstanceId().InstanceIndex,
+		// stops the path at the level of the world the object is in
+		FText::FromString(SMInstance.GetISMComponent()->GetPathName(SMInstance.GetISMComponent()->GetWorld())));
+}
+
+#undef LOCTEXT_NAMESPACE
diff --git a/Engine/Source/Editor/UnrealEd/Private/Elements/SMInstance/SMInstanceElementDetailsProxyObject.h b/Engine/Source/Editor/UnrealEd/Private/Elements/SMInstance/SMInstanceElementDetailsProxyObject.h
--- a/Engine/Source/Editor/UnrealEd/Private/Elements/SMInstance/SMInstanceElementDetailsProxyObject.h
+++ b/Engine/Source/Editor/UnrealEd/Private/Elements/SMInstance/SMInstanceElementDetailsProxyObject.h
@@ -4,6 +4,9 @@
 
 #include "UObject/Object.h"
 #include "Elements/SMInstance/SMInstanceManager.h"
+
+#include "IObjectNameEditSink.h"
+
 #include "SMInstanceElementDetailsProxyObject.generated.h"
 
 UCLASS(Transient)
@@ -31,4 +34,15 @@ private:
 
 	FDelegateHandle TickHandle;
 	bool bIsWithinInteractiveTransformEdit = false;
+
+	friend class FSMInstanceElementDetailsProxyObjectNameEditSink;
 };
+
+class FSMInstanceElementDetailsProxyObjectNameEditSink : public UE::EditorWidgets::IObjectNameEditSink
+{
+	virtual UClass* GetSupportedClass() const override;
+
+	virtual FText GetObjectDisplayName(UObject* Object) const override;
+
+	virtual FText GetObjectNameTooltip(UObject* Object) const override;
+};
\ No newline at end of file
