Added missing '#undef LOCTEXT_NAMESPACE'
#rnx

diff --git a/Engine/Source/Editor/EditorWidgets/Private/SEnumCombo.cpp b/Engine/Source/Editor/EditorWidgets/Private/SEnumCombo.cpp
--- a/Engine/Source/Editor/EditorWidgets/Private/SEnumCombo.cpp
+++ b/Engine/Source/Editor/EditorWidgets/Private/SEnumCombo.cpp
@@ -198,3 +198,5 @@ TSharedRef<SWidget> SEnumComboBox::OnGetMenuContent()
 
 	return MenuBuilder.MakeWidget();
 }
+
+#undef LOCTEXT_NAMESPACE
