Upgraded SAssetDropTarget to deal with drag & drop operations of multiple assets, and deprecated single asset events.

To preserve the old behavior, multiple asset drag drop support is off by default.

diff --git a/Engine/Plugins/2D/Paper2D/Source/Paper2DEditor/Private/TileMapEditing/TileMapEdModeToolkit.cpp b/Engine/Plugins/2D/Paper2D/Source/Paper2DEditor/Private/TileMapEditing/TileMapEdModeToolkit.cpp
--- a/Engine/Plugins/2D/Paper2D/Source/Paper2DEditor/Private/TileMapEditing/TileMapEdModeToolkit.cpp
+++ b/Engine/Plugins/2D/Paper2D/Source/Paper2DEditor/Private/TileMapEditing/TileMapEdModeToolkit.cpp
@@ -96,8 +96,8 @@ void FTileMapEdModeToolkit::Init(const TSharedPtr<IToolkitHost>& InitToolkitHost
 		+SOverlay::Slot()
 		[
 			SNew(SAssetDropTarget)
-			.OnIsAssetAcceptableForDrop(this, &FTileMapEdModeToolkit::OnAssetDraggedOver)
-			.OnAssetDropped(this, &FTileMapEdModeToolkit::OnChangeTileSet)
+			.OnAreAssetsAcceptableForDrop(this, &FTileMapEdModeToolkit::OnAssetDraggedOver)
+			.OnAssetsDropped(this, &FTileMapEdModeToolkit::OnAssetsDropped)
 			[
 				TileSetPalette.ToSharedRef()
 			]
@@ -191,6 +191,11 @@ void FTileMapEdModeToolkit::Init(const TSharedPtr<IToolkitHost>& InitToolkitHost
 	FModeToolkit::Init(InitToolkitHost);
 }
 
+void FTileMapEdModeToolkit::OnAssetsDropped(const FDragDropEvent&, TArrayView<FAssetData> InAssets)
+{
+	OnChangeTileSet(InAssets[0].GetAsset());
+}
+
 void FTileMapEdModeToolkit::OnChangeTileSet(UObject* NewAsset)
 {
 	if (UPaperTileSet* NewTileSet = Cast<UPaperTileSet>(NewAsset))
@@ -398,9 +403,9 @@ FReply FTileMapEdModeToolkit::ClickedOnTileSetPaletteCornerText()
 	return FReply::Handled();
 }
 
-bool FTileMapEdModeToolkit::OnAssetDraggedOver(const UObject* InObject) const
+bool FTileMapEdModeToolkit::OnAssetDraggedOver(TArrayView<FAssetData> InAssets) const
 {
-	return Cast<UPaperTileSet>(InObject) != nullptr;
+	return Cast<UPaperTileSet>(InAssets[0].GetAsset()) != nullptr;
 }
 
 //////////////////////////////////////////////////////////////////////////
diff --git a/Engine/Plugins/2D/Paper2D/Source/Paper2DEditor/Private/TileMapEditing/TileMapEdModeToolkit.h b/Engine/Plugins/2D/Paper2D/Source/Paper2DEditor/Private/TileMapEditing/TileMapEdModeToolkit.h
--- a/Engine/Plugins/2D/Paper2D/Source/Paper2DEditor/Private/TileMapEditing/TileMapEdModeToolkit.h
+++ b/Engine/Plugins/2D/Paper2D/Source/Paper2DEditor/Private/TileMapEditing/TileMapEdModeToolkit.h
@@ -38,6 +38,7 @@ public:
 	// End of FModeToolkit interface
 
 protected:
+	void OnAssetsDropped(const FDragDropEvent&, TArrayView<FAssetData> InAssets);
 	void OnChangeTileSet(UObject* NewAsset);
 	UObject* GetCurrentTileSet() const;
 
@@ -55,7 +56,7 @@ protected:
 	EVisibility GetTileSetPaletteCornerTextVisibility() const;
 	FReply ClickedOnTileSetPaletteCornerText();
 
-	bool OnAssetDraggedOver(const UObject* InObject) const;
+	bool OnAssetDraggedOver(TArrayView<FAssetData> InAssets) const;
 
 private:
 	class FEdModeTileMap* TileMapEditor;
diff --git a/Engine/Source/Editor/DetailCustomizations/Private/DialogueWaveWidgets.cpp b/Engine/Source/Editor/DetailCustomizations/Private/DialogueWaveWidgets.cpp
--- a/Engine/Source/Editor/DetailCustomizations/Private/DialogueWaveWidgets.cpp
+++ b/Engine/Source/Editor/DetailCustomizations/Private/DialogueWaveWidgets.cpp
@@ -43,8 +43,8 @@ void SDialogueVoicePropertyEditor::Construct( const FArguments& InArgs, const TS
 		TSharedRef<SWidget> AssetWidget =
 			SNew( SAssetDropTarget )
 			.ToolTipText( this, &SDialogueVoicePropertyEditor::OnGetToolTip )
-			.OnIsAssetAcceptableForDrop( this, &SDialogueVoicePropertyEditor::OnIsAssetAcceptableForDrop )
-			.OnAssetDropped( this, &SDialogueVoicePropertyEditor::OnAssetDropped )
+			.OnAreAssetsAcceptableForDrop( this, &SDialogueVoicePropertyEditor::OnIsAssetAcceptableForDrop )
+			.OnAssetsDropped( this, &SDialogueVoicePropertyEditor::OnAssetDropped )
 			[
 				SNew( SBox )
 				.WidthOverride( ThumbnailSizeX ) 
@@ -229,15 +229,15 @@ void SDialogueVoicePropertyEditor::CloseMenu()
 	ComboButton->SetIsOpen(false);
 }
 
-bool SDialogueVoicePropertyEditor::OnIsAssetAcceptableForDrop( const UObject* InObject ) const
+bool SDialogueVoicePropertyEditor::OnIsAssetAcceptableForDrop( TArrayView<FAssetData> InAssets ) const
 {
 	// Only dialogue voice can be dropped 
-	return InObject->IsA( UDialogueVoice::StaticClass() );
+	return Cast<UDialogueVoice>( InAssets[0].GetAsset() ) != nullptr;
 }
 
-void SDialogueVoicePropertyEditor::OnAssetDropped( UObject* Object )
+void SDialogueVoicePropertyEditor::OnAssetDropped( const FDragDropEvent&, TArrayView<FAssetData> InAssets )
 {
-	ReplaceDialogueVoice( CastChecked<UDialogueVoice>(Object) );
+	ReplaceDialogueVoice( CastChecked<UDialogueVoice>(InAssets[0].GetAsset()) );
 }
 
 FText SDialogueVoicePropertyEditor::GetDialogueVoiceDescription() const
diff --git a/Engine/Source/Editor/DetailCustomizations/Private/DialogueWaveWidgets.h b/Engine/Source/Editor/DetailCustomizations/Private/DialogueWaveWidgets.h
--- a/Engine/Source/Editor/DetailCustomizations/Private/DialogueWaveWidgets.h
+++ b/Engine/Source/Editor/DetailCustomizations/Private/DialogueWaveWidgets.h
@@ -43,14 +43,14 @@ private:
 	 *
 	 * @param InObject	The asset being dragged over
 	 */
-	bool OnIsAssetAcceptableForDrop( const UObject* InObject ) const;
+	bool OnIsAssetAcceptableForDrop( TArrayView<FAssetData> InAssets ) const;
 
 	/**
 	 * Called when an asset dropped onto the list item
 	 *
 	 * @param InObject	The asset being dropped
 	 */
-	void OnAssetDropped( UObject* Object );
+	void OnAssetDropped( const FDragDropEvent&, TArrayView<FAssetData> InAssets );
 
 	bool CanUseSelectedAsset();
 	void OnUseSelectedDialogueVoice();
diff --git a/Engine/Source/Editor/EditorWidgets/Private/SAssetDropTarget.cpp b/Engine/Source/Editor/EditorWidgets/Private/SAssetDropTarget.cpp
--- a/Engine/Source/Editor/EditorWidgets/Private/SAssetDropTarget.cpp
+++ b/Engine/Source/Editor/EditorWidgets/Private/SAssetDropTarget.cpp
@@ -1,7 +1,6 @@
 // Copyright Epic Games, Inc. All Rights Reserved.
 
 #include "SAssetDropTarget.h"
-#include "AssetData.h"
 #include "DragAndDrop/AssetDragDropOp.h"
 #include "DragAndDrop/ActorDragDropOp.h"
 #include "DragAndDrop/DecoratedDragDropOp.h"
@@ -11,9 +10,10 @@
 
 void SAssetDropTarget::Construct(const FArguments& InArgs )
 {
-	OnAssetDropped = InArgs._OnAssetDropped;
-	OnIsAssetAcceptableForDrop = InArgs._OnIsAssetAcceptableForDrop;
-	OnIsAssetAcceptableForDropWithReason = InArgs._OnIsAssetAcceptableForDropWithReason;
+	OnAssetsDropped = InArgs._OnAssetsDropped;
+	OnAreAssetsAcceptableForDrop = InArgs._OnAreAssetsAcceptableForDrop;
+	OnAreAssetsAcceptableForDropWithReason = InArgs._OnAreAssetsAcceptableForDropWithReason;
+	bSupportsMultiDrop = InArgs._bSupportsMultiDrop;
 
 	SDropTarget::Construct(
 		SDropTarget::FArguments()
@@ -25,29 +25,33 @@ void SAssetDropTarget::Construct(const FArguments& InArgs )
 
 FReply SAssetDropTarget::OnDropped(const FGeometry& InGeometry, const FDragDropEvent& InDragDropEvent)
 {
-	bool bUnused;
-	UObject* Object = GetDroppedObject(InDragDropEvent.GetOperation(), bUnused);
-
-	if ( Object )
+	if (OnAssetsDropped.IsBound())
 	{
-		OnAssetDropped.ExecuteIfBound(Object);
+		bool bRecongnizedEvent = false;
+		TArray<FAssetData> AssetDatas = GetDroppedAssets(InDragDropEvent.GetOperation(), bRecongnizedEvent);
+
+		if (bRecongnizedEvent)
+		{
+			OnAssetsDropped.Execute(InDragDropEvent, AssetDatas);
+			return FReply::Handled();
+		}
 	}
 
-	return FReply::Handled();
+	return FReply::Unhandled();
 }
 
 bool SAssetDropTarget::OnAllowDrop(TSharedPtr<FDragDropOperation> DragDropOperation) const
 {
-	bool bUnused = false;
-	UObject* Object = GetDroppedObject(DragDropOperation, bUnused);
+	bool bRecongnizedEvent = false;
+	TArray<FAssetData> AssetDatas = GetDroppedAssets(DragDropOperation, bRecongnizedEvent);
 
-	if ( Object )
+	if (bRecongnizedEvent)
 	{
 		// Check and see if its valid to drop this object
-		if (OnIsAssetAcceptableForDropWithReason.IsBound())
+		if (OnAreAssetsAcceptableForDropWithReason.IsBound())
 		{
 			FText FailureReason;
-			if (OnIsAssetAcceptableForDropWithReason.Execute(Object, FailureReason))
+			if (OnAreAssetsAcceptableForDropWithReason.Execute(AssetDatas, FailureReason))
 			{
 				return true;
 			}
@@ -68,9 +72,9 @@ bool SAssetDropTarget::OnAllowDrop(TSharedPtr<FDragDropOperation> DragDropOperat
 				return false;
 			}
 		}
-		else if ( OnIsAssetAcceptableForDrop.IsBound() )
+		else if (OnAreAssetsAcceptableForDrop.IsBound())
 		{
-			return OnIsAssetAcceptableForDrop.Execute(Object);
+			return OnAreAssetsAcceptableForDrop.Execute(AssetDatas);
 		}
 		else
 		{
@@ -85,7 +89,7 @@ bool SAssetDropTarget::OnAllowDrop(TSharedPtr<FDragDropOperation> DragDropOperat
 bool SAssetDropTarget::OnIsRecognized(TSharedPtr<FDragDropOperation> DragDropOperation) const
 {
 	bool bRecognizedEvent = false;
-	UObject* Object = GetDroppedObject(DragDropOperation, bRecognizedEvent);
+	GetDroppedAssets(DragDropOperation, bRecognizedEvent);
 
 	return bRecognizedEvent;
 }
@@ -102,57 +106,36 @@ void SAssetDropTarget::OnDragLeave(const FDragDropEvent& DragDropEvent)
 	}
 }
 
-UObject* SAssetDropTarget::GetDroppedObject(TSharedPtr<FDragDropOperation> DragDropOperation, bool& bOutRecognizedEvent) const
+TArray<FAssetData> SAssetDropTarget::GetDroppedAssets(TSharedPtr<FDragDropOperation> DragDropOperation, bool& bOutRecognizedEvent) const
 {
-	bOutRecognizedEvent = false;
-	UObject* DroppedObject = nullptr;
-
+	TArray<FAssetData> DroppedAssets;
 	if (!DragDropOperation)
 	{
-		return DroppedObject;
-	}
-
-	// Asset being dragged from content browser
-	if ( DragDropOperation->IsOfType<FAssetDragDropOp>() )
-	{
-		bOutRecognizedEvent = true;
-		TSharedPtr<FAssetDragDropOp> DragDropOp = StaticCastSharedPtr<FAssetDragDropOp>(DragDropOperation);
-		const TArray<FAssetData>& DroppedAssets = DragDropOp->GetAssets();
-
-		bool bCanDrop = DroppedAssets.Num() == 1;
-
-		if( bCanDrop )
-		{
-			const FAssetData& AssetData = DroppedAssets[0];
-
-			// Make sure the asset is loaded
-			DroppedObject = AssetData.GetAsset();
-		}
+		return DroppedAssets;
 	}
-	// Asset being dragged from some external source
-	else if ( DragDropOperation->IsOfType<FExternalDragOperation>() )
+	
+	if ( DragDropOperation->IsOfType<FActorDragDropOp>() )
 	{
-		TArray<FAssetData> DroppedAssetData = AssetUtil::ExtractAssetDataFromDrag(DragDropOperation);
+		// Handle actors being dragged
+		TSharedPtr<FActorDragDropOp> ActorDragDrop = StaticCastSharedPtr<FActorDragDropOp>(DragDropOperation);
 
-		if (DroppedAssetData.Num() == 1)
+		for (TWeakObjectPtr<AActor> Actor : ActorDragDrop->Actors)
 		{
-			bOutRecognizedEvent = true;
-			DroppedObject = DroppedAssetData[0].GetAsset();
+			FAssetData DroppedActorAsset(Actor.Get());
+			if (DroppedActorAsset.IsValid())
+			{
+				DroppedAssets.Emplace(Actor.Get());
+			}
 		}
 	}
-	// Actor being dragged?
-	else if ( DragDropOperation->IsOfType<FActorDragDropOp>() )
+	else
 	{
-		bOutRecognizedEvent = true;
-		TSharedPtr<FActorDragDropOp> ActorDragDrop = StaticCastSharedPtr<FActorDragDropOp>(DragDropOperation);
-
-		if (ActorDragDrop->Actors.Num() == 1)
-		{
-			DroppedObject = ActorDragDrop->Actors[0].Get();
-		}
+		// Handle assets being dragged
+		DroppedAssets = AssetUtil::ExtractAssetDataFromDrag(DragDropOperation);
 	}
 
-	return DroppedObject;
+	bOutRecognizedEvent = (bSupportsMultiDrop && (DroppedAssets.Num() > 0)) || (DroppedAssets.Num() == 1);
+	return DroppedAssets;
 }
 
 #undef LOCTEXT_NAMESPACE
diff --git a/Engine/Source/Editor/EditorWidgets/Public/SAssetDropTarget.h b/Engine/Source/Editor/EditorWidgets/Public/SAssetDropTarget.h
--- a/Engine/Source/Editor/EditorWidgets/Public/SAssetDropTarget.h
+++ b/Engine/Source/Editor/EditorWidgets/Public/SAssetDropTarget.h
@@ -6,6 +6,7 @@
 #include "Input/Reply.h"
 #include "Widgets/DeclarativeSyntaxSupport.h"
 #include "SDropTarget.h"
+#include "AssetRegistry/AssetData.h"
 
 /**
  * A widget that displays a hover cue and handles dropping assets of allowed types onto this widget
@@ -16,22 +17,75 @@ public:
 	/** Called when a valid asset is dropped */
 	DECLARE_DELEGATE_OneParam( FOnAssetDropped, UObject* );
 
+	/** Called when valid assets are dropped */
+	DECLARE_DELEGATE_TwoParams( FOnAssetsDropped, const FDragDropEvent&, TArrayView<FAssetData> );
+
 	/** Called when we need to check if an asset type is valid for dropping */
 	DECLARE_DELEGATE_RetVal_OneParam( bool, FIsAssetAcceptableForDrop, const UObject* );
 
+	/** Called when we need to check if an asset type is valid for dropping */
+	DECLARE_DELEGATE_RetVal_OneParam( bool, FAreAssetsAcceptableForDrop, TArrayView<FAssetData> );
+
 	/** Called when we need to check if an asset type is valid for dropping and also will have a reason if it is not */
 	DECLARE_DELEGATE_RetVal_TwoParams( bool, FIsAssetAcceptableForDropWithReason, const UObject*, FText& );
 
+	/** Called when we need to check if an asset type is valid for dropping and also will have a reason if it is not */
+	DECLARE_DELEGATE_RetVal_TwoParams( bool, FAreAssetsAcceptableForDropWithReason, TArrayView<FAssetData>, FText& );
+
 	SLATE_BEGIN_ARGS(SAssetDropTarget)
+		: _bSupportsMultiDrop(false)
 	{ }
 		/* Content to display for the in the drop target */
 		SLATE_DEFAULT_SLOT( FArguments, Content )
 		/** Called when a valid asset is dropped */
-		SLATE_EVENT( FOnAssetDropped, OnAssetDropped )
-		/** Called to check if an asset is acceptible for dropping */
-		SLATE_EVENT( FIsAssetAcceptableForDrop, OnIsAssetAcceptableForDrop )
-		/** Called to check if an asset is acceptible for dropping if you also plan on returning a reason text */
-		SLATE_EVENT( FIsAssetAcceptableForDropWithReason, OnIsAssetAcceptableForDropWithReason )
+		SLATE_EVENT( FOnAssetsDropped, OnAssetsDropped )
+		/** Called to check if an asset is acceptable for dropping */
+		SLATE_EVENT( FAreAssetsAcceptableForDrop, OnAreAssetsAcceptableForDrop )
+		/** Called to check if an asset is acceptable for dropping if you also plan on returning a reason text */
+		SLATE_EVENT( FAreAssetsAcceptableForDropWithReason, OnAreAssetsAcceptableForDropWithReason )
+		/** Sets if this drop target can support multiple assets dropped, or only supports a single asset dropped at a time. False by default for legacy behavior. */
+		SLATE_ARGUMENT( bool, bSupportsMultiDrop )
+
+		FOnAssetsDropped ConvertObjectDropDelegate(const FOnAssetDropped& LegacyDelegate)
+		{
+			return FOnAssetsDropped::CreateLambda([LegacyDelegate](const FDragDropEvent&, TArrayView<FAssetData> InAssets)
+			{
+				if (InAssets.Num() == 1)
+				{
+					LegacyDelegate.Execute(InAssets[0].GetAsset());
+				}
+			});
+		}
+		SLATE_EVENT_DEPRECATED(5.0, "Use FOnAssetsDropped instead.", FOnAssetDropped, OnAssetDropped, OnAssetsDropped, ConvertObjectDropDelegate)
+
+		FAreAssetsAcceptableForDrop ConvertObjectAcceptableDropDelegate(const FIsAssetAcceptableForDrop& LegacyDelegate)
+		{
+			return FAreAssetsAcceptableForDrop::CreateLambda([LegacyDelegate](TArrayView<FAssetData> InAssets)
+			{
+				if (InAssets.Num() == 1)
+				{
+					 return LegacyDelegate.Execute(InAssets[0].GetAsset());
+				}
+
+				return false;
+			});
+		}
+		SLATE_EVENT_DEPRECATED(5.0, "Use FAreAssetsAcceptableForDrop instead.", FIsAssetAcceptableForDrop, OnIsAssetAcceptableForDrop, OnAreAssetsAcceptableForDrop, ConvertObjectAcceptableDropDelegate)
+
+		FAreAssetsAcceptableForDropWithReason ConvertObjectAcceptableDropWithReasonDelegate(const FIsAssetAcceptableForDropWithReason& LegacyDelegate)
+		{
+			return FAreAssetsAcceptableForDropWithReason::CreateLambda([LegacyDelegate](TArrayView<FAssetData> InAssets, FText& OutText)
+			{
+				if (InAssets.Num() == 1)
+				{
+					 return LegacyDelegate.Execute(InAssets[0].GetAsset(), OutText);
+				}
+
+				OutText = NSLOCTEXT("SAssetDropTarget", "NoDroppedArgs", "No valid objects were dropped.");
+				return false;
+			});
+		}
+		SLATE_EVENT_DEPRECATED(5.0, "Use FAreAssetsAcceptableForDropWithReason instead.", FIsAssetAcceptableForDropWithReason, OnIsAssetAcceptableForDropWithReason, OnAreAssetsAcceptableForDropWithReason, ConvertObjectAcceptableDropWithReasonDelegate)
 	SLATE_END_ARGS()
 
 	void Construct(const FArguments& InArgs );
@@ -43,13 +97,15 @@ protected:
 	virtual void OnDragLeave(const FDragDropEvent& DragDropEvent) override;
 
 private:
-	UObject* GetDroppedObject(TSharedPtr<FDragDropOperation> DragDropOperation, bool& bOutRecognizedEvent) const;
+	TArray<FAssetData> GetDroppedAssets(TSharedPtr<FDragDropOperation> DragDropOperation, bool& bOutRecognizedEvent) const;
 
 private:
 	/** Delegate to call when an asset is dropped */
-	FOnAssetDropped OnAssetDropped;
+	FOnAssetsDropped OnAssetsDropped;
 	/** Delegate to call to check validity of the asset */
-	FIsAssetAcceptableForDrop OnIsAssetAcceptableForDrop;
+	FAreAssetsAcceptableForDrop OnAreAssetsAcceptableForDrop;
 	/** Delegate to call to check validity of the asset if you will also provide a reason when returning false */
-	FIsAssetAcceptableForDropWithReason OnIsAssetAcceptableForDropWithReason;
+	FAreAssetsAcceptableForDropWithReason OnAreAssetsAcceptableForDropWithReason;
+	/** Boolean to determine if this drop target can support multiple assets dropped, or only supports a single asset dropped at a time. False by default for legacy behavior. */
+	bool bSupportsMultiDrop;
 };
diff --git a/Engine/Source/Editor/LandscapeEditor/Private/LandscapeEditorDetailCustomization_AlphaBrush.cpp b/Engine/Source/Editor/LandscapeEditor/Private/LandscapeEditorDetailCustomization_AlphaBrush.cpp
--- a/Engine/Source/Editor/LandscapeEditor/Private/LandscapeEditorDetailCustomization_AlphaBrush.cpp
+++ b/Engine/Source/Editor/LandscapeEditor/Private/LandscapeEditorDetailCustomization_AlphaBrush.cpp
@@ -204,8 +204,8 @@ void FLandscapeEditorDetailCustomization_AlphaBrush::CustomizeDetails(IDetailLay
 	.MaxDesiredWidth(0)
 	[
 		SNew(SAssetDropTarget)
-		.OnAssetDropped_Static(&FLandscapeEditorDetailCustomization_AlphaBrush::OnAssetDropped, PropertyHandle_AlphaTexture)
-		.OnIsAssetAcceptableForDrop_Static(&FLandscapeEditorDetailCustomization_AlphaBrush::OnAssetDraggedOver)
+		.OnAssetsDropped_Static(&FLandscapeEditorDetailCustomization_AlphaBrush::OnAssetDropped, PropertyHandle_AlphaTexture)
+		.OnAreAssetsAcceptableForDrop_Static(&FLandscapeEditorDetailCustomization_AlphaBrush::OnAssetDraggedOver)
 		.ToolTipText(PropertyHandle_AlphaTexture->GetToolTipText())
 		[
 			SNew(SHorizontalBox)
@@ -266,16 +266,14 @@ void FLandscapeEditorDetailCustomization_AlphaBrush::CustomizeDetails(IDetailLay
 }
 END_SLATE_FUNCTION_BUILD_OPTIMIZATION
 
-bool FLandscapeEditorDetailCustomization_AlphaBrush::OnAssetDraggedOver(const UObject* InObject)
+bool FLandscapeEditorDetailCustomization_AlphaBrush::OnAssetDraggedOver(TArrayView<FAssetData> InAssets)
 {
-	check(InObject);
-
-	return InObject->IsA(UTexture2D::StaticClass());
+	return Cast<UTexture2D>(InAssets[0].GetAsset()) != nullptr;
 }
 
-void FLandscapeEditorDetailCustomization_AlphaBrush::OnAssetDropped(UObject* InObject, TSharedRef<IPropertyHandle> PropertyHandle_AlphaTexture)
+void FLandscapeEditorDetailCustomization_AlphaBrush::OnAssetDropped(const FDragDropEvent&, TArrayView<FAssetData> InAssets, TSharedRef<IPropertyHandle> PropertyHandle_AlphaTexture)
 {
-	ensure(PropertyHandle_AlphaTexture->SetValue(InObject) == FPropertyAccess::Success);
+	ensure(PropertyHandle_AlphaTexture->SetValue(InAssets[0].GetAsset()) == FPropertyAccess::Success);
 }
 
 #undef LOCTEXT_NAMESPACE
diff --git a/Engine/Source/Editor/LandscapeEditor/Private/LandscapeEditorDetailCustomization_AlphaBrush.h b/Engine/Source/Editor/LandscapeEditor/Private/LandscapeEditorDetailCustomization_AlphaBrush.h
--- a/Engine/Source/Editor/LandscapeEditor/Private/LandscapeEditorDetailCustomization_AlphaBrush.h
+++ b/Engine/Source/Editor/LandscapeEditor/Private/LandscapeEditorDetailCustomization_AlphaBrush.h
@@ -25,6 +25,6 @@ public:
 	virtual void CustomizeDetails(IDetailLayoutBuilder& DetailBuilder) override;
 
 protected:
-	static bool OnAssetDraggedOver(const UObject* InObject);
-	static void OnAssetDropped(UObject* InObject, TSharedRef<IPropertyHandle> PropertyHandle_AlphaTexture);
+	static bool OnAssetDraggedOver(TArrayView<FAssetData> InAssets);
+	static void OnAssetDropped(const FDragDropEvent&, TArrayView<FAssetData> InAssets, TSharedRef<IPropertyHandle> PropertyHandle_AlphaTexture);
 };
diff --git a/Engine/Source/Editor/PropertyEditor/Private/UserInterface/PropertyEditor/SPropertyEditorAsset.cpp b/Engine/Source/Editor/PropertyEditor/Private/UserInterface/PropertyEditor/SPropertyEditorAsset.cpp
--- a/Engine/Source/Editor/PropertyEditor/Private/UserInterface/PropertyEditor/SPropertyEditorAsset.cpp
+++ b/Engine/Source/Editor/PropertyEditor/Private/UserInterface/PropertyEditor/SPropertyEditorAsset.cpp
@@ -364,8 +364,8 @@ void SPropertyEditorAsset::Construct(const FArguments& InArgs, const TSharedPtr<
 	ChildSlot
 	[
 		SNew( SAssetDropTarget )
-		.OnIsAssetAcceptableForDropWithReason( this, &SPropertyEditorAsset::OnAssetDraggedOver )
-		.OnAssetDropped( this, &SPropertyEditorAsset::OnAssetDropped )
+		.OnAreAssetsAcceptableForDropWithReason( this, &SPropertyEditorAsset::OnAssetDraggedOver )
+		.OnAssetsDropped( this, &SPropertyEditorAsset::OnAssetDropped )
 		[
 			SAssignNew( ValueContentBox, SHorizontalBox )	
 		]
@@ -1260,12 +1260,13 @@ FSlateColor SPropertyEditorAsset::GetAssetClassColor()
 	return FSlateColor::UseForeground();
 }
 
-bool SPropertyEditorAsset::OnAssetDraggedOver( const UObject* InObject, FText& OutReason ) const
+bool SPropertyEditorAsset::OnAssetDraggedOver( TArrayView<FAssetData> InAssets, FText& OutReason ) const
 {
-	if (CanEdit() && InObject != nullptr && InObject->IsA(ObjectClass))
+	UObject* AssetObject = InAssets[0].GetAsset();
+	if (CanEdit() && (AssetObject != nullptr) && AssetObject->IsA(ObjectClass))
 	{
+		FAssetData AssetData(InAssets[0]);
 		// Check against custom asset filter
-		FAssetData AssetData(InObject);
 		if (!OnShouldFilterAsset.IsBound()
 			|| !OnShouldFilterAsset.Execute(AssetData))
 		{
@@ -1279,11 +1280,11 @@ bool SPropertyEditorAsset::OnAssetDraggedOver( const UObject* InObject, FText& O
 	return false;
 }
 
-void SPropertyEditorAsset::OnAssetDropped( UObject* InObject )
+void SPropertyEditorAsset::OnAssetDropped( const FDragDropEvent&, TArrayView<FAssetData> InAssets )
 {
 	if( CanEdit() )
 	{
-		SetValue(InObject);
+		SetValue(InAssets[0].GetAsset());
 	}
 }
 
diff --git a/Engine/Source/Editor/PropertyEditor/Private/UserInterface/PropertyEditor/SPropertyEditorAsset.h b/Engine/Source/Editor/PropertyEditor/Private/UserInterface/PropertyEditor/SPropertyEditorAsset.h
--- a/Engine/Source/Editor/PropertyEditor/Private/UserInterface/PropertyEditor/SPropertyEditorAsset.h
+++ b/Engine/Source/Editor/PropertyEditor/Private/UserInterface/PropertyEditor/SPropertyEditorAsset.h
@@ -267,13 +267,13 @@ private:
 	 * @param	OutReason	When returning false, the reason it was not allowed
 	 * @returns true if the object can be dropped
 	 */
-	bool OnAssetDraggedOver( const UObject* InObject, FText& OutReason ) const;
+	bool OnAssetDraggedOver( TArrayView<FAssetData> InAssets, FText& OutReason ) const;
 
 	/** 
 	 * Delegate handling dropping an object on this widget
 	 * @param	InObject	The object we are dropping
 	 */
-	void OnAssetDropped( UObject* InObject );
+	void OnAssetDropped( const FDragDropEvent&, TArrayView<FAssetData> InAssets );
 
 	/** 
 	 * Delegate handling ctrl+c
