Moved SEditorHeaderButton out of UnrealEd and into EditorWidgets

diff --git a/Engine/Source/Editor/ClothPainter/ClothPainter.Build.cs b/Engine/Source/Editor/ClothPainter/ClothPainter.Build.cs
--- a/Engine/Source/Editor/ClothPainter/ClothPainter.Build.cs
+++ b/Engine/Source/Editor/ClothPainter/ClothPainter.Build.cs
@@ -33,7 +33,8 @@ public class ClothPainter : ModuleRules
                 "SkeletalMeshEditor",
                 "AdvancedPreviewScene",
 				"ToolMenus",
-            }
+				"EditorWidgets",
+			}
 		);
     }
 }
diff --git a/Engine/Source/Editor/CurveTableEditor/CurveTableEditor.Build.cs b/Engine/Source/Editor/CurveTableEditor/CurveTableEditor.Build.cs
--- a/Engine/Source/Editor/CurveTableEditor/CurveTableEditor.Build.cs
+++ b/Engine/Source/Editor/CurveTableEditor/CurveTableEditor.Build.cs
@@ -19,7 +19,8 @@ public class CurveTableEditor : ModuleRules
 				"Slate",
 				"SlateCore",
                 "EditorStyle",
-				"UnrealEd"
+				"UnrealEd",
+				"EditorWidgets",
 			}
 		);
 
diff --git a/Engine/Source/Editor/EditorWidgets/Private/SEditorHeaderButton.cpp b/Engine/Source/Editor/EditorWidgets/Private/SEditorHeaderButton.cpp
new file mode 100644
--- /dev/null
+++ b/Engine/Source/Editor/EditorWidgets/Private/SEditorHeaderButton.cpp
@@ -0,0 +1,89 @@
+// Copyright Epic Games, Inc. All Rights Reserved.
+
+#include "SEditorHeaderButton.h"
+#include "SComponentClassCombo.h"
+#include "Widgets/Images/SImage.h"
+#include "Widgets/Input/SButton.h"
+
+void SEditorHeaderButton::Construct(const FArguments& InArgs)
+{
+	check(InArgs._Icon.IsSet());
+
+	TAttribute<FText> Text = InArgs._Text;
+	
+	TSharedRef<SHorizontalBox> ButtonContent = SNew(SHorizontalBox)
+		+ SHorizontalBox::Slot()
+		.HAlign(HAlign_Center)
+		.VAlign(VAlign_Center)
+		.AutoWidth()
+		[
+			SNew(SImage)
+			.Image(InArgs._Icon)
+			.ColorAndOpacity(FSlateColor::UseForeground())
+		]
+		+ SHorizontalBox::Slot()
+		.Padding(FMargin(3, 0, 0, 0))
+		.VAlign(VAlign_Center)
+		.AutoWidth()
+		[
+			SNew(STextBlock)
+			.TextStyle(FAppStyle::Get(), "DialogButtonText")
+			.Text(InArgs._Text)
+			.Visibility_Lambda([Text]() { return Text.Get(FText::GetEmpty()).IsEmpty() ? EVisibility::Collapsed : EVisibility::Visible; })
+		];
+
+	if (InArgs._OnClicked.IsBound())
+	{
+		ChildSlot
+		[
+			SAssignNew(Button, SButton)
+			.ButtonStyle(FAppStyle::Get(), "RoundedButton")
+			.ForegroundColor(FSlateColor::UseStyle())
+			.IsEnabled(InArgs._IsEnabled)
+			.ToolTipText(InArgs._ToolTipText)
+			.HAlign(HAlign_Center)
+			.VAlign(VAlign_Center)
+			.ContentPadding(0)
+			.OnClicked(InArgs._OnClicked)
+			[
+				ButtonContent
+			]
+		];
+	}
+	else
+	{
+		ChildSlot
+		[
+			SAssignNew(ComboButton, SComboButton)
+			.ComboButtonStyle(FAppStyle::Get(), "RoundedComboButton")
+			.ForegroundColor(FSlateColor::UseStyle())
+			.IsEnabled(InArgs._IsEnabled)
+			.ToolTipText(InArgs._ToolTipText)
+			.HAlign(HAlign_Center)
+			.VAlign(VAlign_Center)
+			.ButtonContent()
+			[
+				ButtonContent
+			]
+			.MenuContent()
+			[
+				InArgs._MenuContent.Widget
+			]
+			.OnGetMenuContent(InArgs._OnGetMenuContent)
+			.OnMenuOpenChanged(InArgs._OnMenuOpenChanged)
+			.OnComboBoxOpened(InArgs._OnComboBoxOpened)
+		];
+	}
+}
+
+void SEditorHeaderButton::SetMenuContentWidgetToFocus(TWeakPtr<SWidget> Widget)
+{
+	check(ComboButton.IsValid());
+	ComboButton->SetMenuContentWidgetToFocus(Widget);
+}
+
+void SEditorHeaderButton::SetIsMenuOpen(bool bIsOpen, bool bIsFocused)
+{
+	check(ComboButton.IsValid());
+	ComboButton->SetIsOpen(bIsOpen, bIsFocused);
+}
diff --git a/Engine/Source/Editor/EditorWidgets/Public/SEditorHeaderButton.h b/Engine/Source/Editor/EditorWidgets/Public/SEditorHeaderButton.h
new file mode 100644
--- /dev/null
+++ b/Engine/Source/Editor/EditorWidgets/Public/SEditorHeaderButton.h
@@ -0,0 +1,48 @@
+// Copyright Epic Games, Inc. All Rights Reserved.
+
+#pragma once
+
+#include "Framework/SlateDelegates.h"
+#include "Widgets/DeclarativeSyntaxSupport.h"
+#include "Widgets/SCompoundWidget.h"
+#include "Widgets/Input/SComboButton.h"
+
+class EDITORWIDGETS_API SEditorHeaderButton : public SCompoundWidget
+{
+public:
+
+	SLATE_BEGIN_ARGS(SEditorHeaderButton) :
+		_Icon(FAppStyle::Get().GetBrush("Icons.Plus"))
+		{}
+
+		/** The text to display in the button. */
+		SLATE_ATTRIBUTE(FText, Text)
+
+		SLATE_ATTRIBUTE(const FSlateBrush*, Icon)
+
+		/** The clicked handler. Note that if this is set, the button will behave as though it were just a button.
+		 * This means that OnGetMenuContent, OnComboBoxOpened and OnMenuOpenChanged will all be ignored, since there is no menu.
+		 */
+		SLATE_EVENT(FOnClicked, OnClicked)
+
+		/** The static menu content widget. */
+		SLATE_NAMED_SLOT(FArguments, MenuContent)
+
+		SLATE_EVENT(FOnGetContent, OnGetMenuContent)
+		SLATE_EVENT(FOnComboBoxOpened, OnComboBoxOpened)
+		SLATE_EVENT(FOnIsOpenChanged, OnMenuOpenChanged)
+
+	SLATE_END_ARGS()
+
+	SEditorHeaderButton() {}
+
+	void Construct(const FArguments& InArgs);
+
+	void SetMenuContentWidgetToFocus(TWeakPtr<SWidget> Widget);
+	void SetIsMenuOpen(bool bIsOpen, bool bIsFocused);
+
+private:
+
+	TSharedPtr<SComboButton> ComboButton;
+	TSharedPtr<class SButton> Button;
+};
\ No newline at end of file
diff --git a/Engine/Source/Editor/FoliageEdit/FoliageEdit.build.cs b/Engine/Source/Editor/FoliageEdit/FoliageEdit.build.cs
--- a/Engine/Source/Editor/FoliageEdit/FoliageEdit.build.cs
+++ b/Engine/Source/Editor/FoliageEdit/FoliageEdit.build.cs
@@ -27,6 +27,7 @@ public class FoliageEdit : ModuleRules
                 "AssetTools",
                 "Foliage",
 				"DataLayerEditor",
+				"EditorWidgets",
 			}
 		);
 
diff --git a/Engine/Source/Editor/SubobjectEditor/SubobjectEditor.Build.cs b/Engine/Source/Editor/SubobjectEditor/SubobjectEditor.Build.cs
--- a/Engine/Source/Editor/SubobjectEditor/SubobjectEditor.Build.cs
+++ b/Engine/Source/Editor/SubobjectEditor/SubobjectEditor.Build.cs
@@ -31,6 +31,7 @@ public class SubobjectEditor : ModuleRules
 				"UnrealEd",
 				"InputCore",
 				"KismetWidgets",
+				"EditorWidgets",
 				"GameProjectGeneration",
 				"EditorFramework",
 				"GraphEditor",
diff --git a/Engine/Source/Editor/UMGEditor/UMGEditor.Build.cs b/Engine/Source/Editor/UMGEditor/UMGEditor.Build.cs
--- a/Engine/Source/Editor/UMGEditor/UMGEditor.Build.cs
+++ b/Engine/Source/Editor/UMGEditor/UMGEditor.Build.cs
@@ -47,6 +47,7 @@ public class UMGEditor : ModuleRules
 				"EditorFramework",
 				"UnrealEd", // for Asset Editor Subsystem
 				"KismetWidgets",
+				"EditorWidgets",
 				"KismetCompiler",
 				"BlueprintGraph",
 				"GraphEditor",
diff --git a/Engine/Source/Editor/UnrealEd/Private/SEditorHeaderButton.cpp b/Engine/Source/Editor/UnrealEd/Private/SEditorHeaderButton.cpp
deleted file mode 100644
diff --git a/Engine/Source/Editor/UnrealEd/Public/SEditorHeaderButton.h b/Engine/Source/Editor/UnrealEd/Public/SEditorHeaderButton.h
deleted file mode 100644
