Added missing header guards

diff --git a/Engine/Source/Editor/EditorWidgets/Public/SSimpleButton.h b/Engine/Source/Editor/EditorWidgets/Public/SSimpleButton.h
--- a/Engine/Source/Editor/EditorWidgets/Public/SSimpleButton.h
+++ b/Engine/Source/Editor/EditorWidgets/Public/SSimpleButton.h
@@ -1,5 +1,7 @@
 // Copyright Epic Games, Inc. All Rights Reserved.
 
+#pragma once
+
 #include "Framework/SlateDelegates.h"
 #include "Widgets/DeclarativeSyntaxSupport.h"
 #include "Widgets/Input/SButton.h"
diff --git a/Engine/Source/Editor/EditorWidgets/Public/SSimpleComboButton.h b/Engine/Source/Editor/EditorWidgets/Public/SSimpleComboButton.h
--- a/Engine/Source/Editor/EditorWidgets/Public/SSimpleComboButton.h
+++ b/Engine/Source/Editor/EditorWidgets/Public/SSimpleComboButton.h
@@ -1,5 +1,7 @@
 // Copyright Epic Games, Inc. All Rights Reserved.
 
+#pragma once
+
 #include "Framework/SlateDelegates.h"
 #include "Widgets/DeclarativeSyntaxSupport.h"
 #include "Widgets/Input/SComboButton.h"
