Clean up time slider widget code in Rewind Debugger
- new shared SSimpleTimeSlider widget
- double used for time values throughout Rewind Debugger

diff --git a/Engine/Plugins/Animation/GameplayInsights/Source/RewindDebugger/Private/RewindDebugger.cpp b/Engine/Plugins/Animation/GameplayInsights/Source/RewindDebugger/Private/RewindDebugger.cpp
--- a/Engine/Plugins/Animation/GameplayInsights/Source/RewindDebugger/Private/RewindDebugger.cpp
+++ b/Engine/Plugins/Animation/GameplayInsights/Source/RewindDebugger/Private/RewindDebugger.cpp
@@ -491,7 +491,7 @@ void FRewindDebugger::StepBackward()
 }
 
 
-void FRewindDebugger::ScrubToTime(float ScrubTime, bool bIsScrubbing)
+void FRewindDebugger::ScrubToTime(double ScrubTime, bool bIsScrubbing)
 {
 	if (CanScrub())
 	{
@@ -523,7 +523,7 @@ UWorld* FRewindDebugger::GetWorldToVisualize() const
 	return World;
 }
 
-void  FRewindDebugger::SetCurrentScrubTime(float Time)
+void  FRewindDebugger::SetCurrentScrubTime(double Time)
 {
 	CurrentScrubTime = Time;
 	UpdateTraceTime();
diff --git a/Engine/Plugins/Animation/GameplayInsights/Source/RewindDebugger/Private/RewindDebugger.h b/Engine/Plugins/Animation/GameplayInsights/Source/RewindDebugger/Private/RewindDebugger.h
--- a/Engine/Plugins/Animation/GameplayInsights/Source/RewindDebugger/Private/RewindDebugger.h
+++ b/Engine/Plugins/Animation/GameplayInsights/Source/RewindDebugger/Private/RewindDebugger.h
@@ -80,8 +80,8 @@ public:
 	void StepBackward();
 
 	bool CanScrub() const;
-	void ScrubToTime(float ScrubTime, bool bIsScrubbing);
-	float GetScrubTime() { return CurrentScrubTime; }
+	void ScrubToTime(double ScrubTime, bool bIsScrubbing);
+	double GetScrubTime() { return CurrentScrubTime; }
 
 	// Tick function: While recording, update recording duration.  While paused, and we have recorded data, update skinned mesh poses for the current frame, and handle playback.
 	void Tick(float DeltaTime);
@@ -114,7 +114,7 @@ private:
 
 
 	bool UpdateComponentList(uint64 ParentId, TArray<TSharedPtr<FDebugObjectInfo>>& NewComponentList, bool bAddController = false);
-	void SetCurrentScrubTime(float Time);
+	void SetCurrentScrubTime(double Time);
 	void UpdateTraceTime();
 
 	TBindableProperty<double> TraceTime;
@@ -140,7 +140,7 @@ private:
 	bool bRecording;
 
 	float PlaybackRate;
-	float CurrentScrubTime;
+	double CurrentScrubTime;
 	int64 ScrubFrameIndex;
 	uint16 RecordingIndex;
 
diff --git a/Engine/Plugins/Animation/GameplayInsights/Source/RewindDebugger/Private/RewindDebuggerTimeSliderController.cpp b/Engine/Plugins/Animation/GameplayInsights/Source/RewindDebugger/Private/RewindDebuggerTimeSliderController.cpp
deleted file mode 100644
diff --git a/Engine/Plugins/Animation/GameplayInsights/Source/RewindDebugger/Private/RewindDebuggerTimeSliderController.h b/Engine/Plugins/Animation/GameplayInsights/Source/RewindDebugger/Private/RewindDebuggerTimeSliderController.h
deleted file mode 100644
diff --git a/Engine/Plugins/Animation/GameplayInsights/Source/RewindDebugger/Private/SRewindDebugger.cpp b/Engine/Plugins/Animation/GameplayInsights/Source/RewindDebugger/Private/SRewindDebugger.cpp
--- a/Engine/Plugins/Animation/GameplayInsights/Source/RewindDebugger/Private/SRewindDebugger.cpp
+++ b/Engine/Plugins/Animation/GameplayInsights/Source/RewindDebugger/Private/SRewindDebugger.cpp
@@ -12,7 +12,7 @@
 #include "ObjectTrace.h"
 #include "RewindDebuggerStyle.h"
 #include "RewindDebuggerCommands.h"
-#include "SRewindDebuggerTimeSlider.h"
+#include "SSimpleTimeSlider.h"
 #include "SceneOutlinerPublicTypes.h"
 #include "SceneOutlinerModule.h"
 #include "Selection.h"
@@ -33,6 +33,7 @@
 
 SRewindDebugger::SRewindDebugger() 
 	: SCompoundWidget()
+	, ViewRange(0,10)
 	, DebugComponents(nullptr)
 { 
 }
@@ -44,11 +45,11 @@ SRewindDebugger::~SRewindDebugger()
 void SRewindDebugger::TrackCursor(bool bReverse)
 {
 	float ScrubTime = ScrubTimeAttribute.Get();
-	TRange<float> CurrentViewRange = TimeSliderController->GetTimeSliderArgs().ViewRange.Get();
+	TRange<double> CurrentViewRange = ViewRange;
 	float ViewSize = CurrentViewRange.GetUpperBoundValue() - CurrentViewRange.GetLowerBoundValue();
 
-	static const float LeadingEdgeSize = 0.05f;
-	static const float TrailingEdgeThreshold = 0.01f;
+	static const double LeadingEdgeSize = 0.05;
+	static const double TrailingEdgeThreshold = 0.01;
 
 	if(bReverse)
 	{
@@ -79,8 +80,7 @@ void SRewindDebugger::TrackCursor(bool bReverse)
 		}
 	}
 
-
-	TimeSliderController->GetTimeSliderArgs().ViewRange = CurrentViewRange;
+	ViewRange = CurrentViewRange;
 }
 
 void SRewindDebugger::SetDebugTargetActor(AActor* Actor)
@@ -212,28 +212,7 @@ void SRewindDebugger::Construct(const FArguments& InArgs, TSharedRef<FUICommandL
 	}
 	ToolBarBuilder.EndSection();
 
-	FTimeSliderArgs InitArgs;
-	InitArgs.CursorSize = 0.0f;
-	InitArgs.OnScrubPositionChanged = FTimeSliderArgs::FOnScrubPositionChanged::CreateLambda(
-		[this](float NewScrubTime, bool bIsScrubbing)
-		{
-			if (bIsScrubbing)
-			{
-				OnScrubPositionChanged.ExecuteIfBound( NewScrubTime, bIsScrubbing );
-			}
-		}
-	);
-
-	InitArgs.ScrubPosition = ScrubTimeAttribute;
-
-	InitArgs.ClampRange = InitArgs.ClampRange.CreateLambda(
-		[this]()
-		{ 
-			return TRange<float>(0.0f,RecordingDuration.Get());
-		} );
-
-	TimeSliderController = MakeShared<FTimeSliderController>(InitArgs);
-    TSharedRef<STimeSlider> TimeSlider = SNew(STimeSlider, TimeSliderController.ToSharedRef());
+   
 
 	ComponentTreeView =	SNew(SRewindDebuggerComponentTree)
 				.DebugComponents(InArgs._DebugComponents)
@@ -398,7 +377,25 @@ void SRewindDebugger::Construct(const FArguments& InArgs, TSharedRef<FUICommandL
 			SNew(SVerticalBox)
 			+SVerticalBox::Slot().AutoHeight()
 			[
-				TimeSlider
+				SNew(SSimpleTimeSlider)
+					.ClampRangeHighlightSize(0.15f)
+					.ClampRangeHighlightColor(FLinearColor::Red.CopyWithNewOpacity(0.5f))
+					.ScrubPosition(ScrubTimeAttribute)
+					.ViewRange_Lambda([this](){ return ViewRange; })
+					.OnViewRangeChanged_Lambda([this](TRange<double> NewRange) { ViewRange = NewRange; })
+					.ClampRange_Lambda(
+							[this]()
+							{ 
+								return TRange<double>(0.0f,RecordingDuration.Get());
+							})	
+					.OnScrubPositionChanged_Lambda(
+						[this](double NewScrubTime, bool bIsScrubbing)
+								{
+									if (bIsScrubbing)
+									{
+										OnScrubPositionChanged.ExecuteIfBound( NewScrubTime, bIsScrubbing );
+									}
+								})
 			]
 			+SVerticalBox::Slot().FillHeight(1.0f)
 			[
diff --git a/Engine/Plugins/Animation/GameplayInsights/Source/RewindDebugger/Private/SRewindDebugger.h b/Engine/Plugins/Animation/GameplayInsights/Source/RewindDebugger/Private/SRewindDebugger.h
--- a/Engine/Plugins/Animation/GameplayInsights/Source/RewindDebugger/Private/SRewindDebugger.h
+++ b/Engine/Plugins/Animation/GameplayInsights/Source/RewindDebugger/Private/SRewindDebugger.h
@@ -7,7 +7,6 @@
 #include "Framework/Docking/TabManager.h"
 #include "IRewindDebuggerView.h"
 #include "RewindDebuggerModule.h"
-#include "RewindDebuggerTimeSliderController.h"
 #include "SRewindDebuggerComponentTree.h"
 #include "Widgets/DeclarativeSyntaxSupport.h"
 #include "Widgets/Input/SSearchBox.h"
@@ -20,7 +19,7 @@ class SRewindDebugger : public SCompoundWidget
 	typedef TBindablePropertyInitializer<FString, BindingType_Out> DebugTargetInitializer;
 
 public:
-	DECLARE_DELEGATE_TwoParams( FOnScrubPositionChanged, float, bool )
+	DECLARE_DELEGATE_TwoParams( FOnScrubPositionChanged, double, bool )
 	DECLARE_DELEGATE_OneParam( FOnDebugTargetChanged, TSharedPtr<FString> )
 	DECLARE_DELEGATE_OneParam( FOnComponentDoubleClicked, TSharedPtr<FDebugObjectInfo> )
 	DECLARE_DELEGATE_OneParam( FOnComponentSelectionChanged, TSharedPtr<FDebugObjectInfo> )
@@ -31,7 +30,7 @@ public:
 		SLATE_ARGUMENT(DebugTargetInitializer, DebugTargetActor);
 		SLATE_ARGUMENT(TBindablePropertyInitializer<double>, TraceTime);
 		SLATE_ARGUMENT(TBindablePropertyInitializer<float>, RecordingDuration);
-		SLATE_ATTRIBUTE(float, ScrubTime);
+		SLATE_ATTRIBUTE(double, ScrubTime);
 		SLATE_EVENT( FOnScrubPositionChanged, OnScrubPositionChanged);
 		SLATE_EVENT( FBuildComponentContextMenu, BuildComponentContextMenu );
 		SLATE_EVENT( FOnComponentDoubleClicked, OnComponentDoubleClicked );
@@ -59,13 +58,13 @@ public:
 	void TrackCursor(bool bReverse);
 	void RefreshDebugComponents();
 private:
-	TSharedPtr<FTimeSliderController> TimeSliderController;
-
-	TAttribute<float> ScrubTimeAttribute;
+	// Time Slider
+	TAttribute<double> ScrubTimeAttribute;
 	TAttribute<bool> TrackScrubbingAttribute;
 	FOnScrubPositionChanged OnScrubPositionChanged;
-	FBuildComponentContextMenu BuildComponentContextMenu;
-	FOnComponentSelectionChanged OnComponentSelectionChanged;
+	TRange<double> ViewRange;
+	TBindableProperty<double> TraceTime;
+	TBindableProperty<float> RecordingDuration;
 
 	// debug actor selector
 	TSharedRef<SWidget> MakeSelectActorMenu();
@@ -73,8 +72,6 @@ private:
 	FReply OnSelectActorClicked();
 
 	TBindableProperty<FString, BindingType_Out> DebugTargetActor;
-	TBindableProperty<double> TraceTime;
-	TBindableProperty<float> RecordingDuration;
 	TBindableProperty<float> DebugTargetAnimInstanceId;
 
 	void TraceTimeChanged(double Time);
@@ -86,6 +83,8 @@ private:
 	TArray<TSharedPtr<FDebugObjectInfo>>* DebugComponents;
 	TSharedPtr<FDebugObjectInfo> SelectedComponent;
     void ComponentSelectionChanged(TSharedPtr<FDebugObjectInfo> SelectedItem, ESelectInfo::Type SelectInfo);
+	FBuildComponentContextMenu BuildComponentContextMenu;
+	FOnComponentSelectionChanged OnComponentSelectionChanged;
 	
 	TSharedPtr<SRewindDebuggerComponentTree> ComponentTreeView;
 	TSharedPtr<SWidget> OnContextMenuOpening();
diff --git a/Engine/Plugins/Animation/GameplayInsights/Source/RewindDebugger/Private/SRewindDebuggerTimeSlider.cpp b/Engine/Plugins/Animation/GameplayInsights/Source/RewindDebugger/Private/SRewindDebuggerTimeSlider.cpp
deleted file mode 100644
diff --git a/Engine/Plugins/Animation/GameplayInsights/Source/RewindDebugger/Private/SRewindDebuggerTimeSlider.h b/Engine/Plugins/Animation/GameplayInsights/Source/RewindDebugger/Private/SRewindDebuggerTimeSlider.h
deleted file mode 100644
diff --git a/Engine/Source/Editor/EditorWidgets/Private/SSimpleTimeSlider.cpp b/Engine/Source/Editor/EditorWidgets/Private/SSimpleTimeSlider.cpp
new file mode 100644
--- /dev/null
+++ b/Engine/Source/Editor/EditorWidgets/Private/SSimpleTimeSlider.cpp
@@ -0,0 +1,591 @@
+// Copyright Epic Games, Inc. All Rights Reserved.
+
+#include "SSimpleTimeSlider.h"
+#include "EditorStyleSet.h"
+#include "Fonts/FontMeasure.h"
+#include "Framework/Application/SlateApplication.h"
+#include "Rendering/DrawElements.h"
+#include "Styling/CoreStyle.h"
+
+#define LOCTEXT_NAMESPACE "SSimpleTimeSlider"
+
+
+void SSimpleTimeSlider::Construct( const SSimpleTimeSlider::FArguments& InArgs )
+{
+    ScrubPosition = InArgs._ScrubPosition;
+    ViewRange = InArgs._ViewRange;
+	ClampRange = InArgs._ClampRange;
+    AllowZoom = InArgs._AllowZoom;
+    AllowPan = InArgs._AllowPan;
+	CursorSize = InArgs._CursorSize;
+	MirrorLabels = InArgs._MirrorLabels;
+    OnScrubPositionChanged = InArgs._OnScrubPositionChanged;
+    OnViewRangeChanged = InArgs._OnViewRangeChanged;
+	ClampRangeHighlightColor = InArgs._ClampRangeHighlightColor;
+	ClampRangeHighlightSize = InArgs._ClampRangeHighlightSize;
+	
+	DistanceDragged = 0.0f;
+	bDraggingScrubber = false;
+	bPanning = false;
+
+	// set clipping on by default, since the OnPaint function is drawing outside the bounds
+	Clipping = EWidgetClipping::ClipToBounds;
+
+	ScrubHandleUp = FEditorStyle::GetBrush( TEXT( "Sequencer.Timeline.VanillaScrubHandleUp" ) ); 
+	ScrubHandleDown = FEditorStyle::GetBrush( TEXT( "Sequencer.Timeline.VanillaScrubHandleDown" ) );
+	CursorBackground = FEditorStyle::GetBrush("Sequencer.SectionArea.Background");
+}
+
+int32 SSimpleTimeSlider::OnPaint( const FPaintArgs& Args, const FGeometry& AllottedGeometry, const FSlateRect& MyCullingRect, FSlateWindowElementList& OutDrawElements, int32 LayerId, const FWidgetStyle& InWidgetStyle, bool bParentEnabled ) const
+{
+	int32 NewLayer = OnPaintTimeSlider( MirrorLabels.Get(), AllottedGeometry, MyCullingRect, OutDrawElements, LayerId, InWidgetStyle, bParentEnabled );
+
+	return FMath::Max( NewLayer, SCompoundWidget::OnPaint(Args, AllottedGeometry, MyCullingRect, OutDrawElements, NewLayer, InWidgetStyle, ShouldBeEnabled( bParentEnabled ) ) );
+}
+
+FVector2D SSimpleTimeSlider::ComputeDesiredSize( float ) const
+{
+	return FVector2D(100, 22);
+}
+
+namespace ScrubConstants
+{
+	/** The minimum amount of pixels between each major ticks on the widget */
+	const int32 MinPixelsPerDisplayTick = 5;
+
+	/**The smallest number of units between between major tick marks */
+	const float MinDisplayTickSpacing = 0.001f;
+}
+
+
+/** Utility struct for converting between scrub range space and local/absolute screen space */
+struct SSimpleTimeSlider::FScrubRangeToScreen
+{
+	FVector2D WidgetSize;
+
+	TRange<double> ViewInput;
+	float ViewInputRange;
+	float PixelsPerInput;
+
+	FScrubRangeToScreen(TRange<double> InViewInput, const FVector2D& InWidgetSize )
+	{
+		WidgetSize = InWidgetSize;
+
+		ViewInput = InViewInput;
+		ViewInputRange = ViewInput.Size<float>();
+		PixelsPerInput = ViewInputRange > 0 ? ( WidgetSize.X / ViewInputRange ) : 0;
+	}
+
+	/** Local Widget Space -> Curve Input domain. */
+	float LocalXToInput(float ScreenX) const
+	{
+		float LocalX = ScreenX;
+		return (LocalX/PixelsPerInput) + ViewInput.GetLowerBoundValue();
+	}
+
+	/** Curve Input domain -> local Widget Space */
+	float InputToLocalX(float Input) const
+	{
+		return (Input - ViewInput.GetLowerBoundValue()) * PixelsPerInput;
+	}
+};
+
+
+/**
+ * Gets the the next spacing value in the series 
+ * to determine a good spacing value
+ * E.g, .001,.005,.010,.050,.100,.500,1.000,etc
+ */
+static float GetNextSpacing( uint32 CurrentStep )
+{
+	if(CurrentStep & 0x01) 
+	{
+		// Odd numbers
+		return FMath::Pow( 10.f, 0.5f*((float)(CurrentStep-1)) + 1.f );
+	}
+	else 
+	{
+		// Even numbers
+		return 0.5f * FMath::Pow( 10.f, 0.5f*((float)(CurrentStep)) + 1.f );
+	}
+}
+
+/**
+* Determines the optimal spacing between tick marks in the slider for a given pixel density
+* Increments until a minimum amount of slate units specified by MinTick is reached
+*
+* @param InPixelsPerInput	The density of pixels between each input
+* @param MinTick			The minimum slate units per tick allowed
+* @param MinTickSpacing	The minimum tick spacing in time units allowed
+* @return the optimal spacing in time units
+*/
+static double DetermineOptimalSpacing(float InPixelsPerInput, uint32 InMinTick, double InMinTickSpacing)
+{
+	if (InPixelsPerInput == 0.0f)
+		return InMinTickSpacing;
+
+	uint32 CurStep = 0;
+
+	// Start with the smallest spacing
+	double Spacing = InMinTickSpacing;
+	double MinTick = InMinTick;
+
+	while( Spacing * InPixelsPerInput < MinTick )
+	{
+		Spacing = InMinTickSpacing * GetNextSpacing( CurStep );
+		CurStep++;
+	}
+
+	return Spacing;
+}
+
+struct SSimpleTimeSlider::FDrawTickArgs
+{
+	/** Geometry of the area */
+	FGeometry AllottedGeometry;
+	/** Clipping rect of the area */
+	FSlateRect ClippingRect;
+	/** Color of each tick */
+	FLinearColor TickColor;
+	/** Offset in Y where to start the tick */
+	float TickOffset;
+	/** Height in of major ticks */
+	float MajorTickHeight;
+	/** Start layer for elements */
+	int32 StartLayer;
+	/** Draw effects to apply */
+	ESlateDrawEffect DrawEffects;
+	/** Whether or not to only draw major ticks */
+	bool bOnlyDrawMajorTicks;
+	/** Whether or not to mirror labels */
+	bool bMirrorLabels;
+	
+};
+
+void SSimpleTimeSlider::DrawTicks( FSlateWindowElementList& OutDrawElements, const struct FScrubRangeToScreen& RangeToScreen, FDrawTickArgs& InArgs ) const
+{
+	const float Spacing = DetermineOptimalSpacing( RangeToScreen.PixelsPerInput, ScrubConstants::MinPixelsPerDisplayTick, ScrubConstants::MinDisplayTickSpacing );
+
+	// Sub divisions
+	const int32 Divider = 10;
+	// For slightly larger halfway tick mark
+	const int32 HalfDivider = Divider / 2;
+	// Find out where to start from
+	int32 OffsetNum = FMath::FloorToInt(RangeToScreen.ViewInput.GetLowerBoundValue() / Spacing);
+	
+	FSlateFontInfo SmallLayoutFont = FCoreStyle::GetDefaultFontStyle("Regular", 8);
+
+	TArray<FVector2D> LinePoints;
+	LinePoints.AddUninitialized(2);
+
+	double Seconds = 0;
+	while( (Seconds = OffsetNum*Spacing) < RangeToScreen.ViewInput.GetUpperBoundValue() )
+	{
+		// X position local to start of the widget area
+		float XPos = RangeToScreen.InputToLocalX( Seconds );
+		uint32 AbsOffsetNum = FMath::Abs(OffsetNum);
+
+		if ( AbsOffsetNum % Divider == 0 )
+		{
+			FVector2D Offset( XPos, InArgs.TickOffset );
+			FVector2D TickSize( 1.0f, InArgs.MajorTickHeight );
+
+			LinePoints[0] = FVector2D(1.0f,1.0f);
+			LinePoints[1] = TickSize;
+
+			// lines should not need anti-aliasing
+			const bool bAntiAliasLines = false;
+
+			// Draw each tick mark
+			FSlateDrawElement::MakeLines(
+				OutDrawElements,
+				InArgs.StartLayer,
+				InArgs.AllottedGeometry.ToPaintGeometry( Offset, TickSize ),
+				LinePoints,
+				InArgs.DrawEffects,
+				InArgs.TickColor,
+				false
+				);
+
+			if( !InArgs.bOnlyDrawMajorTicks )
+			{
+				FString FrameString = Spacing == ScrubConstants::MinDisplayTickSpacing ? FString::Printf( TEXT("%.3f"), Seconds ) : FString::Printf( TEXT("%.2f"), Seconds );
+
+				// Space the text between the tick mark but slightly above
+				const TSharedRef< FSlateFontMeasure > FontMeasureService = FSlateApplication::Get().GetRenderer()->GetFontMeasureService();
+				FVector2D TextSize = FontMeasureService->Measure(FrameString, SmallLayoutFont);
+				FVector2D TextOffset( XPos-(TextSize.X*0.5f), InArgs.bMirrorLabels ? TextSize.Y :  FMath::Abs( InArgs.AllottedGeometry.GetLocalSize().Y - (InArgs.MajorTickHeight+TextSize.Y) ) );
+
+				FSlateDrawElement::MakeText(
+					OutDrawElements,
+					InArgs.StartLayer+1, 
+					InArgs.AllottedGeometry.ToPaintGeometry( TextOffset, TextSize ), 
+					FrameString, 
+					SmallLayoutFont, 
+					InArgs.DrawEffects,
+					InArgs.TickColor 
+				);
+			}
+		}
+		else if( !InArgs.bOnlyDrawMajorTicks )
+		{
+			// Compute the size of each tick mark.  If we are half way between to visible values display a slightly larger tick mark
+			const float MinorTickHeight = AbsOffsetNum % HalfDivider == 0 ? 7.0f : 4.0f;
+
+			FVector2D Offset(XPos, InArgs.bMirrorLabels ? 0.0f : FMath::Abs( InArgs.AllottedGeometry.GetLocalSize().Y - MinorTickHeight ) );
+			FVector2D TickSize(1, MinorTickHeight);
+
+			LinePoints[0] = FVector2D(1.0f,1.0f);
+			LinePoints[1] = TickSize;
+
+			const bool bAntiAlias = false;
+			// Draw each sub mark
+			FSlateDrawElement::MakeLines(
+				OutDrawElements,
+				InArgs.StartLayer,
+				InArgs.AllottedGeometry.ToPaintGeometry( Offset, TickSize ),
+				LinePoints,
+				InArgs.DrawEffects,
+				InArgs.TickColor,
+				bAntiAlias
+			);
+		}
+		// Advance to next tick mark
+		++OffsetNum;
+	}
+}
+
+int32 SSimpleTimeSlider::OnPaintTimeSlider( bool bMirrorLabels, const FGeometry& AllottedGeometry, const FSlateRect& MyCullingRect, FSlateWindowElementList& OutDrawElements, int32 LayerId, const FWidgetStyle& InWidgetStyle, bool bParentEnabled ) const
+{
+	const bool bEnabled = bParentEnabled;
+	const ESlateDrawEffect DrawEffects = bEnabled ? ESlateDrawEffect::None : ESlateDrawEffect::DisabledEffect;
+
+	TRange<double> LocalViewRange = ViewRange.Get();
+	const float LocalViewRangeMin = LocalViewRange.GetLowerBoundValue();
+	const float LocalViewRangeMax = LocalViewRange.GetUpperBoundValue();
+	const float LocalSequenceLength = LocalViewRangeMax-LocalViewRangeMin;
+	
+	FVector2D Scale = FVector2D(1.0f,1.0f);
+	if ( LocalSequenceLength > 0)
+	{
+		FScrubRangeToScreen RangeToScreen( LocalViewRange, AllottedGeometry.GetLocalSize() );
+		
+		// Draw ClampRange 
+
+		float LeftClamp =  RangeToScreen.InputToLocalX(ClampRange.Get().GetLowerBoundValue());
+		float RightClamp =  RangeToScreen.InputToLocalX(ClampRange.Get().GetUpperBoundValue());
+		float Height = AllottedGeometry.GetLocalSize().Y * ClampRangeHighlightSize.Get();
+
+		FPaintGeometry RangeGeometry;
+		if (bMirrorLabels)
+		{
+			RangeGeometry = AllottedGeometry.ToPaintGeometry(FVector2D(LeftClamp, 0), FVector2D(RightClamp-LeftClamp, Height));
+		}
+		else
+		{
+			RangeGeometry = AllottedGeometry.ToPaintGeometry(FVector2D(LeftClamp, AllottedGeometry.GetLocalSize().Y - Height), FVector2D(RightClamp-LeftClamp, AllottedGeometry.GetLocalSize().Y ));
+		}
+
+		FSlateDrawElement::MakeBox(
+			OutDrawElements,
+			++LayerId,
+			RangeGeometry,
+			CursorBackground,
+			DrawEffects,
+			ClampRangeHighlightColor.Get()
+			);
+	
+		const float MajorTickHeight = 9.0f;
+	
+		FDrawTickArgs Args;
+		Args.AllottedGeometry = AllottedGeometry;
+		Args.bMirrorLabels = bMirrorLabels;
+		Args.bOnlyDrawMajorTicks = false;
+		Args.TickColor = FLinearColor::White;
+		Args.ClippingRect = MyCullingRect;
+		Args.DrawEffects = DrawEffects;
+		Args.StartLayer = LayerId;
+		Args.TickOffset = bMirrorLabels ? 0.0f : FMath::Abs( AllottedGeometry.GetLocalSize().Y - MajorTickHeight );
+		Args.MajorTickHeight = MajorTickHeight;
+
+		DrawTicks( OutDrawElements, RangeToScreen, Args );
+
+		const float HandleSize = 13.0f;
+		float HalfSize = FMath::TruncToFloat(HandleSize/2.0f);
+
+		// Draw the scrub handle
+		const float XPos = RangeToScreen.InputToLocalX( ScrubPosition.Get() );
+
+		// Draw cursor size
+		const float CursorHalfSize = CursorSize.Get() * 0.5f;
+		const int32 CursorLayer = LayerId + 2;
+		const float CursorHalfLength = AllottedGeometry.GetLocalSize().X * CursorHalfSize;
+		FPaintGeometry CursorGeometry = AllottedGeometry.ToPaintGeometry(FVector2D(XPos - CursorHalfLength, 0), FVector2D(2 * CursorHalfLength, AllottedGeometry.GetLocalSize().Y));
+
+		FLinearColor CursorColor = InWidgetStyle.GetColorAndOpacityTint();
+		CursorColor.A = CursorColor.A*0.08f;
+		CursorColor.B *= 0.1f;
+		CursorColor.G *= 0.2f;
+		FSlateDrawElement::MakeBox(
+			OutDrawElements,
+			CursorLayer,
+			CursorGeometry,
+			CursorBackground,
+			DrawEffects,
+			CursorColor
+			);
+		
+
+
+		// Should draw above the text
+		const int32 ArrowLayer = LayerId + 3;
+		FPaintGeometry MyGeometry =	AllottedGeometry.ToPaintGeometry( FVector2D( XPos-HalfSize, 0 ), FVector2D( HandleSize, AllottedGeometry.GetLocalSize().Y ) );
+		FLinearColor ScrubColor = InWidgetStyle.GetColorAndOpacityTint();
+
+		// @todo Sequencer this color should be specified in the style
+		ScrubColor.A = ScrubColor.A*0.5f;
+		ScrubColor.B *= 0.1f;
+		ScrubColor.G *= 0.2f;
+		FSlateDrawElement::MakeBox( 
+			OutDrawElements,
+			ArrowLayer, 
+			MyGeometry,
+			bMirrorLabels ? ScrubHandleUp : ScrubHandleDown,
+			DrawEffects, 
+			ScrubColor
+			);
+
+
+		return ArrowLayer;
+	}
+
+	return LayerId;
+}
+
+FReply SSimpleTimeSlider::OnMouseButtonDown( const FGeometry& MyGeometry, const FPointerEvent& MouseEvent )
+{
+	bool bHandleLeftMouseButton = MouseEvent.GetEffectingButton() == EKeys::LeftMouseButton;
+	bool bHandleRightMouseButton = MouseEvent.GetEffectingButton() == EKeys::RightMouseButton && AllowPan.Get();
+	
+	DistanceDragged = 0;
+
+	if ( bHandleLeftMouseButton )
+	{
+		// Always capture mouse if we left or right click on the widget
+		FScrubRangeToScreen RangeToScreen(ViewRange.Get(), MyGeometry.GetLocalSize());
+		FVector2D CursorPos = MyGeometry.AbsoluteToLocal(MouseEvent.GetLastScreenSpacePosition());
+		float NewValue = RangeToScreen.LocalXToInput(CursorPos.X);
+
+		CommitScrubPosition(NewValue, /*bIsScrubbing=*/false);
+		return FReply::Handled().CaptureMouse(AsShared()).PreventThrottling();
+	}
+	else if ( bHandleRightMouseButton )
+	{
+		return FReply::Handled().CaptureMouse(AsShared());
+	}
+
+	return FReply::Unhandled();
+}
+
+FReply SSimpleTimeSlider::OnMouseButtonUp(const FGeometry& MyGeometry, const FPointerEvent& MouseEvent )
+{
+	bool bHandleLeftMouseButton = MouseEvent.GetEffectingButton() == EKeys::LeftMouseButton && HasMouseCapture();
+	bool bHandleRightMouseButton = MouseEvent.GetEffectingButton() == EKeys::RightMouseButton && HasMouseCapture() && AllowPan.Get();
+	
+	if ( bHandleRightMouseButton )
+	{
+		if (!bPanning)
+		{
+			// return unhandled in case our parent wants to use our right mouse button to open a context menu
+			return FReply::Unhandled().ReleaseMouseCapture();
+		}
+		
+		bPanning = false;
+		return FReply::Handled().ReleaseMouseCapture();
+	}
+	else if ( bHandleLeftMouseButton )
+	{
+		if( !bDraggingScrubber )
+		{
+			FScrubRangeToScreen RangeToScreen( ViewRange.Get(), MyGeometry.GetLocalSize() );
+			FVector2D CursorPos = MyGeometry.AbsoluteToLocal(MouseEvent.GetLastScreenSpacePosition());
+			float NewValue = RangeToScreen.LocalXToInput(CursorPos.X);
+
+			CommitScrubPosition( NewValue, /*bIsScrubbing=*/false );
+		}
+
+		bDraggingScrubber = false;
+		return FReply::Handled().ReleaseMouseCapture();
+
+	}
+
+	return FReply::Unhandled();
+}
+
+float SSimpleTimeSlider::GetTimeAtCursorPosition(const FGeometry& MyGeometry, const FPointerEvent& MouseEvent) const
+{
+	FScrubRangeToScreen RangeToScreen(ViewRange.Get(), MyGeometry.GetLocalSize());
+	FVector2D CursorPos = MyGeometry.AbsoluteToLocal(MouseEvent.GetLastScreenSpacePosition());
+	float NewValue = RangeToScreen.LocalXToInput(CursorPos.X);
+
+	float LocalClampMin = ClampRange.Get().GetLowerBoundValue();
+	float LocalClampMax = ClampRange.Get().GetUpperBoundValue();
+
+	if (NewValue < LocalClampMin)
+	{
+		NewValue = LocalClampMin;
+	}
+
+	if (NewValue > LocalClampMax)
+	{
+		NewValue = LocalClampMax;
+	}
+
+	return NewValue;
+}
+
+FReply SSimpleTimeSlider::OnMouseMove( const FGeometry& MyGeometry, const FPointerEvent& MouseEvent )
+{
+	if ( HasMouseCapture() )
+	{
+		if (MouseEvent.IsMouseButtonDown(EKeys::RightMouseButton))
+		{
+			if (!bPanning)
+			{
+				DistanceDragged += FMath::Abs( MouseEvent.GetCursorDelta().X );
+				if ( DistanceDragged > FSlateApplication::Get().GetDragTriggerDistance() )
+				{
+					FReply::Handled().CaptureMouse(AsShared()).UseHighPrecisionMouseMovement(AsShared());
+					SoftwareCursorPosition = MyGeometry.AbsoluteToLocal(MouseEvent.GetLastScreenSpacePosition());
+					bPanning = true;
+				}
+			}
+			else
+			{
+				SoftwareCursorPosition = MyGeometry.AbsoluteToLocal(MouseEvent.GetLastScreenSpacePosition());
+
+				TRange<double> LocalViewRange = ViewRange.Get();
+				double LocalViewRangeMin = LocalViewRange.GetLowerBoundValue();
+				float LocalViewRangeMax = LocalViewRange.GetUpperBoundValue();
+
+				FScrubRangeToScreen ScaleInfo( LocalViewRange, MyGeometry.GetLocalSize() );
+				FVector2D ScreenDelta = MouseEvent.GetCursorDelta();
+				FVector2D InputDelta;
+				InputDelta.X = ScreenDelta.X/ScaleInfo.PixelsPerInput;
+
+				const TRange<double> NewViewRange = TRange<double>(LocalViewRangeMin - InputDelta.X, LocalViewRangeMax - InputDelta.X);
+				TRange<double> LocalClampRange = ClampRange.Get();
+
+				OnViewRangeChanged.ExecuteIfBound(NewViewRange);
+
+				if (!ViewRange.IsBound())
+				{
+					// The  output is not bound to a delegate so we'll manage the value ourselves
+					ViewRange.Set(NewViewRange);
+				}
+			}
+		}
+		else if (MouseEvent.IsMouseButtonDown( EKeys::LeftMouseButton ))
+		{
+			if ( !bDraggingScrubber )
+			{
+				DistanceDragged += FMath::Abs( MouseEvent.GetCursorDelta().X );
+				if ( DistanceDragged > 0/*FSlateApplication::Get().GetDragTriggerDistance()*/ )
+				{
+					bDraggingScrubber = true;
+				}
+			}
+			else
+			{
+				const float NewValue = GetTimeAtCursorPosition(MyGeometry, MouseEvent);
+				CommitScrubPosition(NewValue, /*bIsScrubbing=*/true);
+			}
+		}
+		return FReply::Handled();
+	}
+
+	return FReply::Unhandled();
+}
+
+void SSimpleTimeSlider::CommitScrubPosition( float NewValue, bool bIsScrubbing )
+{
+	// Manage the scrub position ourselves if its not bound to a delegate
+	if ( !ScrubPosition.IsBound() )
+	{
+		ScrubPosition.Set( NewValue );
+	}
+
+	if (!ViewRange.IsBound())
+	{
+		TRange<double> LocalViewRange = ViewRange.Get();
+		const double RangeSize = LocalViewRange.Size<double>();
+		if (NewValue < LocalViewRange.GetLowerBoundValue())
+		{
+			SetTimeRange(NewValue, NewValue + RangeSize);
+		}
+		else if (NewValue > LocalViewRange.GetUpperBoundValue())
+		{
+			SetTimeRange(NewValue - RangeSize, NewValue);
+		}
+	}
+
+	OnScrubPositionChanged.ExecuteIfBound( NewValue, bIsScrubbing );
+}
+
+void SSimpleTimeSlider::SetTimeRange(double NewViewOutputMin, double NewViewOutputMax)
+{
+	if (!ViewRange.IsBound())
+	{
+		ViewRange.Set(TRange<double>(NewViewOutputMin, NewViewOutputMax));
+	}
+	
+	OnViewRangeChanged.ExecuteIfBound(TRange<double>(NewViewOutputMin, NewViewOutputMax));
+}
+
+void SSimpleTimeSlider::SetClampRange(double MinValue, double MaxValue)
+{
+	if (!ClampRange.IsBound())
+	{
+		ClampRange.Set(TRange<double>(MinValue, MaxValue));
+	}
+}
+
+FReply SSimpleTimeSlider::OnMouseWheel( const FGeometry& MyGeometry, const FPointerEvent& MouseEvent )
+{
+	FReply ReturnValue = FReply::Unhandled();
+
+	if ( AllowZoom.Get() )
+	{
+		const float ZoomDelta = -0.1f * MouseEvent.GetWheelDelta();
+
+		float MouseFractionX = MyGeometry.AbsoluteToLocal(MouseEvent.GetScreenSpacePosition()).X / MyGeometry.GetLocalSize().X;
+
+		TRange<double> LocalViewRange = ViewRange.Get();
+		double LocalViewRangeMax = LocalViewRange.GetUpperBoundValue();
+		double LocalViewRangeMin = LocalViewRange.GetLowerBoundValue();
+		const double OutputViewSize = LocalViewRangeMax - LocalViewRangeMin;
+		const double OutputChange = OutputViewSize * ZoomDelta;
+
+		double NewViewOutputMin = LocalViewRangeMin - (OutputChange * MouseFractionX);
+		double NewViewOutputMax = LocalViewRangeMax + (OutputChange * (1.0f - MouseFractionX));
+
+		if (NewViewOutputMin < NewViewOutputMax)
+		{
+			double LocalClampMin = ClampRange.Get().GetLowerBoundValue();
+			double LocalClampMax = ClampRange.Get().GetUpperBoundValue();
+
+			OnViewRangeChanged.ExecuteIfBound(TRange<double>(NewViewOutputMin, NewViewOutputMax));
+			if( !ViewRange.IsBound() )
+			{	
+				// The  output is not bound to a delegate so we'll manage the value ourselves
+				ViewRange.Set( TRange<double>( NewViewOutputMin, NewViewOutputMax ) );
+			}
+		}
+
+		ReturnValue = FReply::Handled();
+	}
+
+	return ReturnValue;
+}
+
+#undef LOCTEXT_NAMESPACE
diff --git a/Engine/Source/Editor/EditorWidgets/Public/SSimpleTimeSlider.h b/Engine/Source/Editor/EditorWidgets/Public/SSimpleTimeSlider.h
new file mode 100644
--- /dev/null
+++ b/Engine/Source/Editor/EditorWidgets/Public/SSimpleTimeSlider.h
@@ -0,0 +1,153 @@
+// Copyright Epic Games, Inc. All Rights Reserved.
+
+#pragma once
+
+#include "CoreMinimal.h"
+#include "Input/Reply.h"
+#include "Widgets/SCompoundWidget.h"
+#include "Widgets/DeclarativeSyntaxSupport.h"
+#include "Widgets/Layout/SScrollBar.h"
+
+class FPaintArgs;
+class FSlateWindowElementList;
+class FTimeSliderController;
+
+class EDITORWIDGETS_API SSimpleTimeSlider : public SCompoundWidget
+{
+public:
+	DECLARE_DELEGATE_OneParam( FOnRangeChanged, TRange<double> )
+	DECLARE_DELEGATE_TwoParams( FOnScrubPositionChanged, double, bool )
+
+	SLATE_BEGIN_ARGS(SSimpleTimeSlider)
+		: _MirrorLabels( false )
+		, _ScrubPosition(0)
+		, _ViewRange(TRange<double>(0,10))
+		, _ClampRange(TRange<double>(0,10))
+		, _AllowZoom (true)
+		, _AllowPan (true)
+		, _CursorSize(0)
+		, _ClampRangeHighlightColor(FLinearColor(0.05,0.05,0.05,1.0))
+		, _ClampRangeHighlightSize(1.0)
+	{}
+		/* If we should mirror the labels on the timeline */
+		SLATE_ATTRIBUTE( bool, MirrorLabels )
+	
+		/** The scrub position */
+    	SLATE_ATTRIBUTE(double, ScrubPosition);
+    
+    	/** View time range */
+    	SLATE_ATTRIBUTE(TRange<double>, ViewRange);
+    
+    	/** Clamp time range */
+    	SLATE_ATTRIBUTE(TRange<double>, ClampRange);
+    
+    	/** If we are allowed to zoom */
+    	SLATE_ATTRIBUTE(bool, AllowZoom);
+	
+    	/** If we are allowed to pan */
+    	SLATE_ATTRIBUTE(bool, AllowPan);
+    
+    	/** Cursor range for data like histogram graphs, etc. */
+    	SLATE_ATTRIBUTE(float, CursorSize);
+	
+    	/** Color overlay for active range */
+		SLATE_ATTRIBUTE(FLinearColor, ClampRangeHighlightColor);
+
+		/** Size of clamp range overlay 0.0 for none 1.0 for the height of the slider */
+		SLATE_ATTRIBUTE(float, ClampRangeHighlightSize);
+    
+	
+		/** Called when the scrub position changes */
+		SLATE_EVENT(FOnScrubPositionChanged, OnScrubPositionChanged);
+	
+		/** Called right before the scrubber begins to move */
+		SLATE_EVENT(FSimpleDelegate, OnBeginScrubberMovement);
+	
+		/** Called right after the scrubber handle is released by the user */
+		SLATE_EVENT(FSimpleDelegate, OnEndScrubberMovement);
+	
+    	/** Called when the view range changes */
+    	SLATE_EVENT(FOnRangeChanged, OnViewRangeChanged);
+     
+	SLATE_END_ARGS()
+
+
+	/**
+	 * Construct the widget
+	 * 
+	 * @param InArgs   A declaration from which to construct the widget
+	 */
+	void Construct( const FArguments& InArgs );
+
+	TRange<double> GetTimeRange() { return ViewRange.Get(); }
+	void SetTimeRange(double MinValue, double MaxValue);
+	void SetClampRange(double MinValue, double MaxValue);
+	bool IsPanning() { return bPanning; }
+
+protected:
+	// forward declared as class members to prevent name collision with similar types defined in other units
+	struct FScrubRangeToScreen;
+	struct FDrawTickArgs;
+
+	// SWidget interface
+	virtual FVector2D ComputeDesiredSize(float) const override;
+	virtual int32 OnPaint( const FPaintArgs& Args, const FGeometry& AllottedGeometry, const FSlateRect& MyCullingRect, FSlateWindowElementList& OutDrawElements, int32 LayerId, const FWidgetStyle& InWidgetStyle, bool bParentEnabled ) const override;
+	virtual FReply OnMouseButtonDown( const FGeometry& MyGeometry, const FPointerEvent& MouseEvent ) override;
+	virtual FReply OnMouseButtonUp( const FGeometry& MyGeometry, const FPointerEvent& MouseEvent ) override;
+	virtual FReply OnMouseMove( const FGeometry& MyGeometry, const FPointerEvent& MouseEvent ) override;
+	virtual FReply OnMouseWheel( const FGeometry& MyGeometry, const FPointerEvent& MouseEvent ) override;
+	
+	int32 OnPaintTimeSlider( bool bMirrorLabels, const FGeometry& AllottedGeometry, const FSlateRect& MyCullingRect, FSlateWindowElementList& OutDrawElements, int32 LayerId, const FWidgetStyle& InWidgetStyle, bool bParentEnabled ) const;
+	
+	/**
+	 * Draws time tick marks
+	 *
+	 * @param OutDrawElements	List to add draw elements to
+	 * @param RangeToScreen		Time range to screen space converter
+	 * @param InArgs			Parameters for drawing the tick lines
+	 */
+	void DrawTicks( FSlateWindowElementList& OutDrawElements, const FScrubRangeToScreen& RangeToScreen, FDrawTickArgs& InArgs ) const;
+
+	float GetTimeAtCursorPosition(const FGeometry& MyGeometry, const FPointerEvent& MouseEvent) const;
+	
+	/**
+	 * Call this method when the user's interaction has changed the scrub position
+	 *
+	 * @param NewValue				Value resulting from the user's interaction
+	 * @param bIsScrubbing			True if done via scrubbing, false if just releasing scrubbing
+	 */
+	void CommitScrubPosition( float NewValue, bool bIsScrubbing );
+
+	TAttribute<double> ScrubPosition;
+	TAttribute<TRange<double>> ViewRange;
+	TAttribute<TRange<double>> ClampRange;
+	TAttribute<double> TimeSnapInterval;
+	TAttribute<bool> AllowZoom;
+	TAttribute<bool> AllowPan;
+	TAttribute<float> CursorSize;
+	TAttribute<FLinearColor> ClampRangeHighlightColor;
+	TAttribute<float> ClampRangeHighlightSize;
+	TAttribute<bool> MirrorLabels;
+	
+	FOnScrubPositionChanged OnScrubPositionChanged;
+	FOnRangeChanged OnViewRangeChanged;
+
+	/** Brush for drawing an upwards facing scrub handle */
+	const FSlateBrush* ScrubHandleUp;
+	/** Brush for drawing a downwards facing scrub handle */
+	const FSlateBrush* ScrubHandleDown;
+	/** Brush for drawing cursor background to visualize cursor size */
+	const FSlateBrush* CursorBackground;
+	/** Total mouse delta during dragging **/
+	float DistanceDragged;
+	/** If we are dragging the scrubber */
+	bool bDraggingScrubber;
+	/** If we are currently panning the panel */
+	bool bPanning;
+	/***/
+	TSharedPtr<SScrollBar> Scrollbar;
+	FVector2D SoftwareCursorPosition;
+
+	
+};
+
