Localization (advanced text settings) dropdown reskin for STextPropertyEditableTextBox

- Convert grid layout to row layout similar to property editor in appearance
- Preserve source string value when unlinking string table - fixes issue with being unable to unlink string table for enum display names
- Combine combo button with culture-invariant warning icon - this also prevents menu from potentially shifting when localizable toggled

diff --git a/Engine/Source/Editor/EditorWidgets/Private/STextPropertyEditableTextBox.cpp b/Engine/Source/Editor/EditorWidgets/Private/STextPropertyEditableTextBox.cpp
--- a/Engine/Source/Editor/EditorWidgets/Private/STextPropertyEditableTextBox.cpp
+++ b/Engine/Source/Editor/EditorWidgets/Private/STextPropertyEditableTextBox.cpp
@@ -23,6 +23,7 @@
 #include "Internationalization/StringTableCore.h"
 #include "Internationalization/StringTableRegistry.h"
 #include "Serialization/TextReferenceCollector.h"
+#include "Styling/StyleColors.h"
 
 #define LOCTEXT_NAMESPACE "STextPropertyEditableTextBox"
 
@@ -96,6 +97,7 @@ void STextPropertyEditableStringTableReference::Construct(const FArguments& InAr
 	TSharedRef<SHorizontalBox> HorizontalBox = SNew(SHorizontalBox);
 
 	HorizontalBox->AddSlot()
+		.Padding(0)
 		[
 			SAssignNew(StringTableOptionsCombo, SComboButton)
 			.ComboButtonStyle(&InArgs._ComboStyle->ComboButtonStyle)
@@ -107,10 +109,12 @@ void STextPropertyEditableStringTableReference::Construct(const FArguments& InAr
 				SNew(STextBlock)
 				.Text(this, &STextPropertyEditableStringTableReference::GetStringTableComboContent)
 				.ToolTipText(this, &STextPropertyEditableStringTableReference::GetStringTableComboToolTip)
+				.Font(InArgs._Font)
 			]
 		];
 
 	HorizontalBox->AddSlot()
+		.Padding(10, 0)
 		[
 			SAssignNew(StringTableKeysCombo, SComboButton)
 			.ComboButtonStyle(&InArgs._ComboStyle->ComboButtonStyle)
@@ -123,6 +127,7 @@ void STextPropertyEditableStringTableReference::Construct(const FArguments& InAr
 				SNew(STextBlock)
 				.Text(this, &STextPropertyEditableStringTableReference::GetKeyComboContent)
 				.ToolTipText(this, &STextPropertyEditableStringTableReference::GetKeyComboToolTip)
+				.Font(InArgs._Font)
 			]
 		];
 
@@ -130,14 +135,26 @@ void STextPropertyEditableStringTableReference::Construct(const FArguments& InAr
 	{
 		HorizontalBox->AddSlot()
 			.AutoWidth()
+			.VAlign(VAlign_Center)
 			[
-				SNew(SButton)
-				.ButtonStyle(InArgs._ButtonStyle)
+				SNew(SBox)
 				.HAlign(HAlign_Center)
 				.VAlign(VAlign_Center)
-				.Text(LOCTEXT("UnlinkStringTable", "Unlink"))
-				.IsEnabled(this, &STextPropertyEditableStringTableReference::IsUnlinkEnabled)
-				.OnClicked(this, &STextPropertyEditableStringTableReference::OnUnlinkClicked)
+				.HeightOverride(22)
+				.WidthOverride(22)
+				[
+					SNew(SButton)
+					.ButtonStyle(InArgs._ButtonStyle)
+					.ContentPadding(0)
+					.ToolTipText(LOCTEXT("UnlinkStringTable", "Unlink"))
+					.IsEnabled(this, &STextPropertyEditableStringTableReference::IsUnlinkEnabled)
+					.OnClicked(this, &STextPropertyEditableStringTableReference::OnUnlinkClicked)
+					[
+						SNew(SImage)
+						.Image(FEditorStyle::GetBrush("Icons.Delete"))
+						.ColorAndOpacity(FSlateColor::UseForeground())
+					]
+				]
 			];
 	}
 
@@ -556,13 +573,136 @@ FReply STextPropertyEditableStringTableReference::OnUnlinkClicked()
 		const FText CurrentText = EditableTextProperty->GetText(TextIndex);
 		if (CurrentText.IsFromStringTable())
 		{
-			EditableTextProperty->SetText(TextIndex, FText::GetEmpty());
+			// Make a copy of the FText separate from the string table but generate a new stable namespace and key
+			// This prevents problems with properties that disallow empty text (e.g. enum display name)
+			FString NewNamespace;
+			FString NewKey;
+			EditableTextProperty->GetStableTextId(
+				TextIndex,
+				IEditableTextProperty::ETextPropertyEditAction::EditedKey,
+				CurrentText.ToString(),
+				FString(),
+				FString(),
+				NewNamespace,
+				NewKey
+			);
+			
+			EditableTextProperty->SetText(TextIndex, FText::ChangeKey(NewNamespace, NewKey, CurrentText));
 		}
 	}
 
 	return FReply::Handled();
 }
 
+/** Single row in the advanced text settings/localization menu. Has a similar appearance to a details row in the property editor. */
+class STextPropertyEditableOptionRow : public SCompoundWidget
+{
+	SLATE_BEGIN_ARGS(STextPropertyEditableOptionRow)
+		: _IsHeader(false)
+		, _ContentHAlign(HAlign_Fill)
+		{}
+		SLATE_ARGUMENT(bool, IsHeader)
+		SLATE_ARGUMENT(EHorizontalAlignment, ContentHAlign)
+		SLATE_ATTRIBUTE(FText, Text)
+		SLATE_DEFAULT_SLOT(FArguments, Content)
+	SLATE_END_ARGS()
+
+public:
+	void Construct(const FArguments& InArgs, TSharedRef<FLinkedBoxManager> InManager)
+	{
+		InArgs._Content.Widget->SetToolTip(GetToolTip());
+
+		if (InArgs._IsHeader)
+		{
+			// Header row, text only, fills entire row
+			ChildSlot
+			[
+				SNew(SBorder)
+				.BorderImage(FAppStyle::Get().GetBrush("DetailsView.GridLine"))
+				.Padding(FMargin(0, 0, 0, 1))
+				[
+					SNew(SBorder)
+					.BorderImage(FAppStyle::Get().GetBrush("DetailsView.CategoryTop"))
+					.BorderBackgroundColor(FSlateColor(FLinearColor::White))
+					.Padding(FMargin(12, 8, 0, 8))
+					.VAlign(VAlign_Center)
+					[
+						SNew(STextBlock)
+						.TransformPolicy(ETextTransformPolicy::ToUpper)
+						.TextStyle(FAppStyle::Get(), "DetailsView.CategoryTextStyle")
+						.Font(FAppStyle::Get().GetFontStyle("PropertyWindow.BoldFont"))
+						.Text(InArgs._Text)
+						.ToolTip(GetToolTip())
+					]
+				]
+			];
+		}
+		else
+		{
+			// Non-header row, has a name column followed by a value widget
+			ChildSlot
+			[
+				SNew(SHorizontalBox)
+				+ SHorizontalBox::Slot()
+				.AutoWidth()
+				[
+					SNew(SLinkedBox, InManager)
+					[
+						SNew(SBorder)
+						.BorderImage(FAppStyle::Get().GetBrush("DetailsView.GridLine"))
+						.Padding(FMargin(0, 0, 0, 1))
+						[
+							SNew(SBorder)
+							.BorderImage(FAppStyle::Get().GetBrush("DetailsView.CategoryMiddle"))
+							.BorderBackgroundColor(this, &STextPropertyEditableOptionRow::GetBackgroundColor)
+							.Padding(FMargin(20, 3.5, 0, 3.5))
+							.HAlign(HAlign_Left)
+							.VAlign(VAlign_Center)
+							[
+								SNew(STextBlock)
+								.Font(FAppStyle::Get().GetFontStyle("PropertyWindow.NormalFont"))
+								.Text(InArgs._Text)
+								.ToolTip(GetToolTip())
+							]
+						]
+					]
+				]
+				+ SHorizontalBox::Slot()
+				[
+					SNew(SBorder)
+					.BorderImage(FAppStyle::Get().GetBrush("DetailsView.GridLine"))
+					.Padding(FMargin(0, 0, 0, 1))
+					[
+						SNew(SBorder)
+						.BorderImage(FAppStyle::Get().GetBrush("DetailsView.CategoryMiddle"))
+						.BorderBackgroundColor(this, &STextPropertyEditableOptionRow::GetBackgroundColor)
+						.Padding(FMargin(14, 3.5, 4, 3.5))
+						.HAlign(InArgs._ContentHAlign)
+						.VAlign(VAlign_Center)
+						[
+							InArgs._Content.Widget
+						]
+					]
+				]
+			];
+
+			// Clear the tooltip from this widget since it's set on the name/value widgets now
+			SetToolTip(nullptr);
+		}
+	}
+
+private:
+	FSlateColor GetBackgroundColor() const
+	{
+		if (IsHovered())
+		{
+			return FStyleColors::Header;
+		}
+
+		return FStyleColors::Panel;
+	}
+};
+
 void STextPropertyEditableTextBox::Construct(const FArguments& InArgs, const TSharedRef<IEditableTextProperty>& InEditableTextProperty)
 {
 	EditableTextProperty = InEditableTextProperty;
@@ -611,7 +751,6 @@ void STextPropertyEditableTextBox::Construct(const FArguments& InArgs, const TSh
 		[
 			SAssignNew(HorizontalBox, SHorizontalBox)
 			+SHorizontalBox::Slot()
-			.VAlign(VAlign_Center)
 			[
 				SNew(SBox)
 				.MinDesiredWidth(InArgs._MinDesiredWidth)
@@ -636,6 +775,9 @@ void STextPropertyEditableTextBox::Construct(const FArguments& InArgs, const TSh
 		PrimaryWidget = SingleLineWidget;
 	}
 
+	const TSharedRef<FLinkedBoxManager> LinkedBoxManager = MakeShared<FLinkedBoxManager>();
+	const FSlateFontInfo PropertyNormalFont = FAppStyle::Get().GetFontStyle("PropertyWindow.NormalFont");
+
 	HorizontalBox->AddSlot()
 		.AutoWidth()
 		[
@@ -645,221 +787,128 @@ void STextPropertyEditableTextBox::Construct(const FArguments& InArgs, const TSh
 			.ContentPadding(FMargin(4, 0))
 			.ButtonStyle(FEditorStyle::Get(), "HoverHintOnly")
 			.ForegroundColor(FSlateColor::UseForeground())
-			.ToolTipText(LOCTEXT("AdvancedTextSettingsComboToolTip", "Edit advanced text settings."))
+			.ToolTipText(this, &STextPropertyEditableTextBox::GetAdvancedTextSettingsComboToolTip)
 			.HasDownArrow(false)
 			.ButtonContent()
 			[
 				SNew(SImage)
-				.Image(FAppStyle::Get().GetBrush("LocalizationDashboard.MenuIcon"))
+				.Image(this, &STextPropertyEditableTextBox::GetAdvancedTextSettingsComboImage)
 				.ColorAndOpacity(FSlateColor::UseForeground())
 			]
 			.MenuContent()
 			[
 				SNew(SBox)
 				.WidthOverride(340)
-				.Padding(4)
+				.Padding(1)
 				[
-					SNew(SGridPanel)
-					.FillColumn(1, 1.0f)
-
-					// Inline Text
-					+SGridPanel::Slot(0, 0)
-					.ColumnSpan(2)
-					.Padding(2)
-					.HAlign(HAlign_Left)
-					.VAlign(VAlign_Center)
+					SNew(SVerticalBox)
+					+ SVerticalBox::Slot()
 					[
-						SNew(STextBlock)
-						.TextStyle(FEditorStyle::Get(), "LargeText")
-						.Text(LOCTEXT("TextInlineTextLabel", "Inline Text"))
+						SNew(STextPropertyEditableOptionRow, LinkedBoxManager)
+						.Text(LOCTEXT("TextLocalizableLabel", "Localize"))
+						.ToolTipText(LOCTEXT("TextLocalizableCheckBoxToolTip", "Whether to assign this text a key and allow it to be gathered for localization.\nIf set to false, marks this text as 'culture invariant' to prevent it being gathered for localization."))
+						.ContentHAlign(HAlign_Left)
+						[
+							SNew(SCheckBox)
+							.IsEnabled(this, &STextPropertyEditableTextBox::IsCultureInvariantFlagEnabled)
+							.IsChecked(this, &STextPropertyEditableTextBox::GetLocalizableCheckState)
+							.OnCheckStateChanged(this, &STextPropertyEditableTextBox::HandleLocalizableCheckStateChanged)
+						]
 					]
-
-					// Localizable?
-					+SGridPanel::Slot(0, 1)
-					.Padding(2)
-					.HAlign(HAlign_Right)
-					.VAlign(VAlign_Center)
+					+ SVerticalBox::Slot()
 					[
-						SNew(STextBlock)
-						.Text(LOCTEXT("TextLocalizableLabel", "Localizable:"))
+						SNew(STextPropertyEditableOptionRow, LinkedBoxManager)
+						.IsHeader(true)
+						.Text(LOCTEXT("TextReferencedTextLabel", "Referenced Text"))
 					]
-					+SGridPanel::Slot(1, 1)
-					.Padding(2)
+					+ SVerticalBox::Slot()
 					[
-						SNew(SHorizontalBox)
-
-						+SHorizontalBox::Slot()
-						.AutoWidth()
-						.Padding(0)
+						SNew(STextPropertyEditableOptionRow, LinkedBoxManager)
+						.Text(LOCTEXT("TextStringTableLabel", "String Table"))
+						.IsEnabled(this, &STextPropertyEditableTextBox::IsTextLocalizable)
 						[
-							SNew(SUniformGridPanel)
-							.SlotPadding(FMargin(0, 0, 4, 0))
-
-							+SUniformGridPanel::Slot(0, 0)
-							[
-								SNew(SCheckBox)
-								.Style(FEditorStyle::Get(), "ToggleButtonCheckbox")
-								.ToolTipText(LOCTEXT("TextLocalizableToggleYesToolTip", "Assign this text a key and allow it to be gathered for localization."))
-								.Padding(FMargin(4, 2))
-								.HAlign(HAlign_Center)
-								.IsEnabled(this, &STextPropertyEditableTextBox::IsCultureInvariantFlagEnabled)
-								.IsChecked(this, &STextPropertyEditableTextBox::GetLocalizableCheckState, true/*bActiveState*/)
-								.OnCheckStateChanged(this, &STextPropertyEditableTextBox::HandleLocalizableCheckStateChanged, true/*bActiveState*/)
-								[
-									SNew(STextBlock)
-									.Text(LOCTEXT("TextLocalizableToggleYes", "Yes"))
-								]
-							]
-
-							+SUniformGridPanel::Slot(1, 0)
-							[
-								SNew(SCheckBox)
-								.Style(FEditorStyle::Get(), "ToggleButtonCheckbox")
-								.ToolTipText(LOCTEXT("TextLocalizableToggleNoToolTip", "Mark this text as 'culture invariant' to prevent it being gathered for localization."))
-								.Padding(FMargin(4, 2))
-								.HAlign(HAlign_Center)
-								.IsEnabled(this, &STextPropertyEditableTextBox::IsCultureInvariantFlagEnabled)
-								.IsChecked(this, &STextPropertyEditableTextBox::GetLocalizableCheckState, false/*bActiveState*/)
-								.OnCheckStateChanged(this, &STextPropertyEditableTextBox::HandleLocalizableCheckStateChanged, false/*bActiveState*/)
-								[
-									SNew(STextBlock)
-									.Text(LOCTEXT("TextLocalizableToggleNo", "No"))
-								]
-							]
+							SNew(STextPropertyEditableStringTableReference, InEditableTextProperty)
+							.AllowUnlink(true)
+							.Font(PropertyNormalFont)
+							.IsEnabled(this, &STextPropertyEditableTextBox::CanEdit)
 						]
 					]
-
-#if USE_STABLE_LOCALIZATION_KEYS
-					// Package
-					+SGridPanel::Slot(0, 2)
-					.Padding(2)
-					.HAlign(HAlign_Right)
-					.VAlign(VAlign_Center)
-					[
-						SNew(STextBlock)
-						.Text(LOCTEXT("TextPackageLabel", "Package:"))
-						.Visibility(this, &STextPropertyEditableTextBox::GetLocalizableVisibility)
-					]
-					+SGridPanel::Slot(1, 2)
-					.Padding(2)
+					+ SVerticalBox::Slot()
 					[
-						SNew(SEditableTextBox)
-						.Text(this, &STextPropertyEditableTextBox::GetPackageValue)
-						.Visibility(this, &STextPropertyEditableTextBox::GetLocalizableVisibility)
-						.IsReadOnly(true)
+						SNew(STextPropertyEditableOptionRow, LinkedBoxManager)
+						.IsHeader(true)
+						.Text(LOCTEXT("TextInlineTextLabel", "Inline Text"))
 					]
-#endif // USE_STABLE_LOCALIZATION_KEYS
 
-					// Namespace
-					+SGridPanel::Slot(0, 3)
-					.Padding(2)
-					.HAlign(HAlign_Right)
-					.VAlign(VAlign_Center)
-					[
-						SNew(STextBlock)
-						.Text(LOCTEXT("TextNamespaceLabel", "Namespace:"))
-						.Visibility(this, &STextPropertyEditableTextBox::GetLocalizableVisibility)
-					]
-					+SGridPanel::Slot(1, 3)
-					.Padding(2)
+#if USE_STABLE_LOCALIZATION_KEYS
+					+ SVerticalBox::Slot()
 					[
-						SAssignNew(NamespaceEditableTextBox, SEditableTextBox)
-						.Text(this, &STextPropertyEditableTextBox::GetNamespaceValue)
-						.SelectAllTextWhenFocused(true)
-						.ClearKeyboardFocusOnCommit(false)
-						.OnTextChanged(this, &STextPropertyEditableTextBox::OnNamespaceChanged)
-						.OnTextCommitted(this, &STextPropertyEditableTextBox::OnNamespaceCommitted)
-						.SelectAllTextOnCommit(true)
-						.Visibility(this, &STextPropertyEditableTextBox::GetLocalizableVisibility)
-						.IsReadOnly(this, &STextPropertyEditableTextBox::IsIdentityReadOnly)
+						SNew(STextPropertyEditableOptionRow, LinkedBoxManager)
+						.Text(LOCTEXT("TextPackageLabel", "Package"))
+						.IsEnabled(this, &STextPropertyEditableTextBox::IsTextLocalizable)
+						[
+							SNew(SEditableTextBox)
+							.Text(this, &STextPropertyEditableTextBox::GetPackageValue)
+							.Font(PropertyNormalFont)
+							.IsReadOnly(true)
+						]
 					]
+#endif // USE_STABLE_LOCALIZATION_KEYS
 
-					// Key
-					+SGridPanel::Slot(0, 4)
-					.Padding(2)
-					.HAlign(HAlign_Right)
-					.VAlign(VAlign_Center)
+					+ SVerticalBox::Slot()
 					[
-						SNew(STextBlock)
-						.Text(LOCTEXT("TextKeyLabel", "Key:"))
-						.Visibility(this, &STextPropertyEditableTextBox::GetLocalizableVisibility)
+						SNew(STextPropertyEditableOptionRow, LinkedBoxManager)
+						.Text(LOCTEXT("TextNamespaceLabel", "Namespace"))
+						.IsEnabled(this, &STextPropertyEditableTextBox::IsTextLocalizable)
+						[
+							SAssignNew(NamespaceEditableTextBox, SEditableTextBox)
+							.Text(this, &STextPropertyEditableTextBox::GetNamespaceValue)
+							.Font(PropertyNormalFont)
+							.SelectAllTextWhenFocused(true)
+							.ClearKeyboardFocusOnCommit(false)
+							.OnTextChanged(this, &STextPropertyEditableTextBox::OnNamespaceChanged)
+							.OnTextCommitted(this, &STextPropertyEditableTextBox::OnNamespaceCommitted)
+							.SelectAllTextOnCommit(true)
+							.IsReadOnly(this, &STextPropertyEditableTextBox::IsIdentityReadOnly)
+						]
 					]
-					+SGridPanel::Slot(1, 4)
-					.Padding(2)
+					+ SVerticalBox::Slot()
 					[
-						SAssignNew(KeyEditableTextBox, SEditableTextBox)
-						.Text(this, &STextPropertyEditableTextBox::GetKeyValue)
-						.Visibility(this, &STextPropertyEditableTextBox::GetLocalizableVisibility)
+						SNew(STextPropertyEditableOptionRow, LinkedBoxManager)
+						.Text(LOCTEXT("TextKeyLabel", "Key"))
+						.IsEnabled(this, &STextPropertyEditableTextBox::IsTextLocalizable)
+						[
+							SAssignNew(KeyEditableTextBox, SEditableTextBox)
+							.Text(this, &STextPropertyEditableTextBox::GetKeyValue)
+							.Font(PropertyNormalFont)
 #if USE_STABLE_LOCALIZATION_KEYS
-						.SelectAllTextWhenFocused(true)
-						.ClearKeyboardFocusOnCommit(false)
-						.OnTextChanged(this, &STextPropertyEditableTextBox::OnKeyChanged)
-						.OnTextCommitted(this, &STextPropertyEditableTextBox::OnKeyCommitted)
-						.SelectAllTextOnCommit(true)
-						.IsReadOnly(this, &STextPropertyEditableTextBox::IsIdentityReadOnly)
+							.SelectAllTextWhenFocused(true)
+							.ClearKeyboardFocusOnCommit(false)
+							.OnTextChanged(this, &STextPropertyEditableTextBox::OnKeyChanged)
+							.OnTextCommitted(this, &STextPropertyEditableTextBox::OnKeyCommitted)
+							.SelectAllTextOnCommit(true)
+							.IsReadOnly(this, &STextPropertyEditableTextBox::IsIdentityReadOnly)
 #else	// USE_STABLE_LOCALIZATION_KEYS
-						.IsReadOnly(true)
+							.IsReadOnly(true)
 #endif	// USE_STABLE_LOCALIZATION_KEYS
-					]
-
-					// Referenced Text
-					+SGridPanel::Slot(0, 5)
-					.ColumnSpan(2)
-					.Padding(2)
-					.HAlign(HAlign_Left)
-					.VAlign(VAlign_Center)
-					[
-						SNew(STextBlock)
-						.TextStyle(FEditorStyle::Get(), "LargeText")
-						.Text(LOCTEXT("TextReferencedTextLabel", "Referenced Text"))
-						.Visibility(this, &STextPropertyEditableTextBox::GetLocalizableVisibility)
-					]
-
-					// String Table
-					+SGridPanel::Slot(0, 6)
-					.Padding(2)
-					.HAlign(HAlign_Right)
-					.VAlign(VAlign_Center)
-					[
-						SNew(STextBlock)
-						.Text(LOCTEXT("TextStringTableLabel", "String Table:"))
-						.Visibility(this, &STextPropertyEditableTextBox::GetLocalizableVisibility)
-					]
-					+SGridPanel::Slot(1, 6)
-					.Padding(2)
-					[
-						SNew(STextPropertyEditableStringTableReference, InEditableTextProperty)
-						.AllowUnlink(true)
-						.IsEnabled(this, &STextPropertyEditableTextBox::CanEdit)
-						.Visibility(this, &STextPropertyEditableTextBox::GetLocalizableVisibility)
+						]
 					]
 				]
 			]
 		];
 
-	HorizontalBox->AddSlot()
-		.VAlign(VAlign_Center)
-		.HAlign(HAlign_Center)
-		.AutoWidth()
-		[
-			SNew(SImage)
-			.Image(FCoreStyle::Get().GetBrush("Icons.Warning"))
-			.Visibility(this, &STextPropertyEditableTextBox::GetTextWarningImageVisibility)
-			.ToolTipText(LOCTEXT("TextNotLocalizedWarningToolTip", "This text is marked as 'culture invariant' and won't be gathered for localization.\nYou can change this by editing the advanced text settings."))
-		];
-
 	SetEnabled(TAttribute<bool>(this, &STextPropertyEditableTextBox::CanEdit));
 }
 
-EVisibility STextPropertyEditableTextBox::GetLocalizableVisibility() const
+bool STextPropertyEditableTextBox::IsTextLocalizable() const
 {
 	const int32 NumTexts = EditableTextProperty->GetNumTexts();
 	if (NumTexts == 1)
 	{
 		const FText PropertyValue = EditableTextProperty->GetText(0);
-		return PropertyValue.IsCultureInvariant() ? EVisibility::Collapsed : EVisibility::Visible;
+		return !PropertyValue.IsCultureInvariant();
 	}
-	return EVisibility::Visible;
+	return true;
 }
 
 void STextPropertyEditableTextBox::GetDesiredWidth(float& OutMinDesiredWidth, float& OutMaxDesiredWidth)
@@ -1373,7 +1422,7 @@ FText STextPropertyEditableTextBox::GetPackageValue() const
 
 #endif // USE_STABLE_LOCALIZATION_KEYS
 
-ECheckBoxState STextPropertyEditableTextBox::GetLocalizableCheckState(bool bActiveState) const
+ECheckBoxState STextPropertyEditableTextBox::GetLocalizableCheckState() const
 {
 	const int32 NumTexts = EditableTextProperty->GetNumTexts();
 	if (NumTexts == 1)
@@ -1381,17 +1430,17 @@ ECheckBoxState STextPropertyEditableTextBox::GetLocalizableCheckState(bool bActi
 		const FText PropertyValue = EditableTextProperty->GetText(0);
 
 		const bool bIsLocalized = !PropertyValue.IsCultureInvariant();
-		return bIsLocalized == bActiveState ? ECheckBoxState::Checked : ECheckBoxState::Unchecked;
+		return bIsLocalized ? ECheckBoxState::Checked : ECheckBoxState::Unchecked;
 	}
 
 	return ECheckBoxState::Unchecked;
 }
 
-void STextPropertyEditableTextBox::HandleLocalizableCheckStateChanged(ECheckBoxState InCheckboxState, bool bActiveState)
+void STextPropertyEditableTextBox::HandleLocalizableCheckStateChanged(ECheckBoxState InCheckboxState)
 {
 	const int32 NumTexts = EditableTextProperty->GetNumTexts();
 
-	if (bActiveState)
+	if (InCheckboxState == ECheckBoxState::Checked)
 	{
 		for (int32 TextIndex = 0; TextIndex < NumTexts; ++TextIndex)
 		{
@@ -1433,17 +1482,28 @@ void STextPropertyEditableTextBox::HandleLocalizableCheckStateChanged(ECheckBoxS
 	}
 }
 
-EVisibility STextPropertyEditableTextBox::GetTextWarningImageVisibility() const
+FText STextPropertyEditableTextBox::GetAdvancedTextSettingsComboToolTip() const
 {
-	const int32 NumTexts = EditableTextProperty->GetNumTexts();
-	
-	if (NumTexts == 1)
+	if (IsTextLocalizable())
 	{
-		const FText PropertyValue = EditableTextProperty->GetText(0);
-		return PropertyValue.IsCultureInvariant() ? EVisibility::Visible : EVisibility::Collapsed;
+		return LOCTEXT("AdvancedTextSettingsComboToolTip", "Edit advanced text settings.");
+	}
+	else
+	{
+		return LOCTEXT("TextNotLocalizedWarningToolTip", "This text is marked as 'culture invariant' and won't be gathered for localization.\nYou can change this by editing the advanced text settings.");
 	}
+}
 
-	return EVisibility::Collapsed;
+const FSlateBrush* STextPropertyEditableTextBox::GetAdvancedTextSettingsComboImage() const
+{
+	if (IsTextLocalizable())
+	{
+		return FAppStyle::Get().GetBrush("LocalizationDashboard.MenuIcon");
+	}
+	else
+	{
+		return FCoreStyle::Get().GetBrush("Icons.Warning");
+	}
 }
 
 bool STextPropertyEditableTextBox::IsValidIdentity(const FText& InIdentity, FText* OutReason, const FText* InErrorCtx) const
diff --git a/Engine/Source/Editor/EditorWidgets/Public/STextPropertyEditableTextBox.h b/Engine/Source/Editor/EditorWidgets/Public/STextPropertyEditableTextBox.h
--- a/Engine/Source/Editor/EditorWidgets/Public/STextPropertyEditableTextBox.h
+++ b/Engine/Source/Editor/EditorWidgets/Public/STextPropertyEditableTextBox.h
@@ -86,13 +86,15 @@ class EDITORWIDGETS_API STextPropertyEditableStringTableReference : public SComp
 {
 	SLATE_BEGIN_ARGS(STextPropertyEditableStringTableReference)
 		: _ComboStyle(&FCoreStyle::Get().GetWidgetStyle<FComboBoxStyle>("ComboBox"))
-		, _ButtonStyle(&FCoreStyle::Get().GetWidgetStyle<FButtonStyle>("Button"))
+		, _ButtonStyle(&FAppStyle::Get().GetWidgetStyle<FButtonStyle>("SimpleButton"))
 		, _AllowUnlink(false)
 		{}
 		/** The styling of the combobox */
 		SLATE_STYLE_ARGUMENT(FComboBoxStyle, ComboStyle)
 		/** The styling of the button */
 		SLATE_STYLE_ARGUMENT(FButtonStyle, ButtonStyle)
+		/** Font for comboboxes */
+		SLATE_ARGUMENT(FSlateFontInfo, Font)
 		/** Should we show an "unlink" button? */
 		SLATE_ARGUMENT(bool, AllowUnlink)
 	SLATE_END_ARGS()
@@ -191,7 +193,7 @@ private:
 	bool IsSourceTextReadOnly() const;
 	bool IsIdentityReadOnly() const;
 	FText GetToolTipText() const;
-	EVisibility GetLocalizableVisibility() const;
+	bool IsTextLocalizable() const;
 
 	FText GetTextValue() const;
 	void OnTextChanged(const FText& NewText);
@@ -210,11 +212,12 @@ private:
 	FText GetPackageValue() const;
 #endif // USE_STABLE_LOCALIZATION_KEYS
 
-	ECheckBoxState GetLocalizableCheckState(bool bActiveState) const;
+	ECheckBoxState GetLocalizableCheckState() const;
 
-	void HandleLocalizableCheckStateChanged(ECheckBoxState InCheckboxState, bool bActiveState);
+	void HandleLocalizableCheckStateChanged(ECheckBoxState InCheckboxState);
 
-	EVisibility GetTextWarningImageVisibility() const;
+	FText GetAdvancedTextSettingsComboToolTip() const;
+	const FSlateBrush* GetAdvancedTextSettingsComboImage() const;
 
 	bool IsValidIdentity(const FText& InIdentity, FText* OutReason = nullptr, const FText* InErrorCtx = nullptr) const;
 
