Fix SWidget that are not incrementing their LayerId correctly.

diff --git a/Engine/Source/Editor/EditorWidgets/Private/SObjectNameEditableTextBox.cpp b/Engine/Source/Editor/EditorWidgets/Private/SObjectNameEditableTextBox.cpp
--- a/Engine/Source/Editor/EditorWidgets/Private/SObjectNameEditableTextBox.cpp
+++ b/Engine/Source/Editor/EditorWidgets/Private/SObjectNameEditableTextBox.cpp
@@ -126,7 +126,7 @@ int32 SObjectNameEditableTextBox::OnPaint( const FPaintArgs& Args, const FGeomet
 			HighlightTargetColorAndOpacity );							// Color
 	}
 
-	return LayerId + TextLayer;
+	return FMath::Max(StartLayer, LayerId + TextLayer);
 }
 
 FText SObjectNameEditableTextBox::GetNameText() const
diff --git a/Engine/Source/Editor/Persona/Private/SPoseWatchManagerTreeView.cpp b/Engine/Source/Editor/Persona/Private/SPoseWatchManagerTreeView.cpp
--- a/Engine/Source/Editor/Persona/Private/SPoseWatchManagerTreeView.cpp
+++ b/Engine/Source/Editor/Persona/Private/SPoseWatchManagerTreeView.cpp
@@ -374,7 +374,7 @@ int32 SPoseWatchManagerTreeRow::OnPaint(const FPaintArgs& Args, const FGeometry&
 			HighlightTargetColorAndOpacity);							// Color
 	}
 
-	return LayerId + TextLayer;
+	return FMath::Max(StartLayer, LayerId + TextLayer);
 }
 
 #undef LOCTEXT_NAMESPACE
diff --git a/Engine/Source/Editor/SceneOutliner/Private/SOutlinerTreeView.cpp b/Engine/Source/Editor/SceneOutliner/Private/SOutlinerTreeView.cpp
--- a/Engine/Source/Editor/SceneOutliner/Private/SOutlinerTreeView.cpp
+++ b/Engine/Source/Editor/SceneOutliner/Private/SOutlinerTreeView.cpp
@@ -395,7 +395,7 @@ int32 SSceneOutlinerTreeRow::OnPaint( const FPaintArgs& Args, const FGeometry& A
 			HighlightTargetColorAndOpacity );							// Color
 	}
 
-	return StartLayer + TextLayer;
+	return FMath::Max(StartLayer, LayerId + TextLayer);
 }
 
 #undef LOCTEXT_NAMESPACE
