SEnumCombo now uses the cached tooltip and display name values instead of fetching them again every tick.

diff --git a/Engine/Source/Editor/EditorWidgets/Private/SEnumCombo.cpp b/Engine/Source/Editor/EditorWidgets/Private/SEnumCombo.cpp
--- a/Engine/Source/Editor/EditorWidgets/Private/SEnumCombo.cpp
+++ b/Engine/Source/Editor/EditorWidgets/Private/SEnumCombo.cpp
@@ -118,7 +118,7 @@ FText SEnumComboBox::GetCurrentValueText() const
 	}
 
 	const int32 ValueNameIndex = Enum->GetIndexByValue(CurrentValue.Get());
-	return Enum->GetDisplayNameTextByIndex(ValueNameIndex);
+	return VisibleEnums[ValueNameIndex].DisplayName;
 }
 
 FText SEnumComboBox::GetCurrentValueTooltip() const
@@ -156,7 +156,7 @@ FText SEnumComboBox::GetCurrentValueTooltip() const
 	{
 		return OnGetToolTipForValue.Execute(ValueNameIndex);
 	}
-	return Enum->GetToolTipTextByIndex(ValueNameIndex);
+	return VisibleEnums[ValueNameIndex].TooltipText;
 }
 
 TSharedRef<SWidget> SEnumComboBox::OnGetMenuContent()
