Reverting SEnumCombo changes in 19111919 due to causing crashes when UMETA(Hidden) values were at the start of the enum.

diff --git a/Engine/Source/Editor/EditorWidgets/Private/SEnumCombo.cpp b/Engine/Source/Editor/EditorWidgets/Private/SEnumCombo.cpp
--- a/Engine/Source/Editor/EditorWidgets/Private/SEnumCombo.cpp
+++ b/Engine/Source/Editor/EditorWidgets/Private/SEnumCombo.cpp
@@ -118,7 +118,7 @@ FText SEnumComboBox::GetCurrentValueText() const
 	}
 
 	const int32 ValueNameIndex = Enum->GetIndexByValue(CurrentValue.Get());
-	return VisibleEnums[ValueNameIndex].DisplayName;
+	return Enum->GetDisplayNameTextByIndex(ValueNameIndex);
 }
 
 FText SEnumComboBox::GetCurrentValueTooltip() const
