Reverted change that broke SEnumCombos with hidden values.

diff --git a/Engine/Source/Editor/EditorWidgets/Private/SEnumCombo.cpp b/Engine/Source/Editor/EditorWidgets/Private/SEnumCombo.cpp
--- a/Engine/Source/Editor/EditorWidgets/Private/SEnumCombo.cpp
+++ b/Engine/Source/Editor/EditorWidgets/Private/SEnumCombo.cpp
@@ -156,7 +156,7 @@ FText SEnumComboBox::GetCurrentValueTooltip() const
 	{
 		return OnGetToolTipForValue.Execute(ValueNameIndex);
 	}
-	return VisibleEnums[ValueNameIndex].TooltipText;
+	return Enum->GetToolTipTextByIndex(ValueNameIndex);
 }
 
 TSharedRef<SWidget> SEnumComboBox::OnGetMenuContent()
