Fixed issue where string table selection combobox wouldn't show the filter text box anymore if there were no results, making it impossible to clear

See UDN: https://epicgames.lightning.force.com/lightning/r/Case/5004z00001dnzLDAAY/view

diff --git a/Engine/Source/Editor/EditorWidgets/Private/STextPropertyEditableTextBox.cpp b/Engine/Source/Editor/EditorWidgets/Private/STextPropertyEditableTextBox.cpp
--- a/Engine/Source/Editor/EditorWidgets/Private/STextPropertyEditableTextBox.cpp
+++ b/Engine/Source/Editor/EditorWidgets/Private/STextPropertyEditableTextBox.cpp
@@ -192,38 +192,43 @@ TSharedRef<SWidget> STextPropertyEditableStringTableReference::OnGetStringTableC
 		.BorderImage(&ComboButtonStyle.MenuBorderBrush)
 		.Padding(ComboButtonStyle.MenuBorderPadding)
 		[
+			SNew(SVerticalBox)
+			+SVerticalBox::Slot()
+			.AutoHeight()
+			[
+				SAssignNew(OptionsSearchBox, SSearchBox)
+				.OnTextChanged(this, &STextPropertyEditableStringTableReference::OnOptionsFilterTextChanged)
+			]
+			+SVerticalBox::Slot()
+			.AutoHeight()
+			[
 			SNew(SWidgetSwitcher)
 			.WidgetIndex_Lambda([this]() { return StringTableComboOptions.IsEmpty() ? 0 : 1; })
 
-			+ SWidgetSwitcher::Slot() // Appears when there are no string tables with keys
-			.Padding(12)
-			[
-				SNew(STextBlock).Text(LOCTEXT("EmptyStringTableList", "No string tables available"))
-			]
+				+SWidgetSwitcher::Slot() // Appears when there are no string tables with keys
+				.Padding(12)
+				[
+					SNew(STextBlock).Text(LOCTEXT("EmptyStringTableList", "No string tables available"))
+				]
 
-			+ SWidgetSwitcher::Slot() // Appears when there's a string table with at least a key
-			[
-				SNew(SBox)
-				.Padding(4)
-				.WidthOverride(280)
-				.MaxDesiredHeight(600)
+				+SWidgetSwitcher::Slot() // Appears when there's a string table with at least a key
 				[
-					SNew(SVerticalBox)
-					+SVerticalBox::Slot()
-					.AutoHeight()
+					SNew(SBox)
+					.Padding(4)
+					.WidthOverride(280)
+					.MaxDesiredHeight(600)
 					[
-						SAssignNew(OptionsSearchBox, SSearchBox)
-						.OnTextChanged(this, &STextPropertyEditableStringTableReference::OnOptionsFilterTextChanged)
-					]
-					+SVerticalBox::Slot()
-					.FillHeight(1.f)
-					.Padding(0, 5, 0, 0)
-					[
-						SAssignNew(StringTableOptionsList, SListView<TSharedPtr<FAvailableStringTable>>)
-						.ListItemsSource(&StringTableComboOptions)
-						.SelectionMode(ESelectionMode::Single)
-						.OnGenerateRow(this, &STextPropertyEditableStringTableReference::OnGenerateStringTableComboOption)
-						.OnSelectionChanged(this, &STextPropertyEditableStringTableReference::OnStringTableComboChanged)
+						SNew(SVerticalBox)
+						+SVerticalBox::Slot()
+						.FillHeight(1.f)
+						.Padding(0, 5, 0, 0)
+						[
+							SAssignNew(StringTableOptionsList, SListView<TSharedPtr<FAvailableStringTable>>)
+							.ListItemsSource(&StringTableComboOptions)
+							.SelectionMode(ESelectionMode::Single)
+							.OnGenerateRow(this, &STextPropertyEditableStringTableReference::OnGenerateStringTableComboOption)
+							.OnSelectionChanged(this, &STextPropertyEditableStringTableReference::OnStringTableComboChanged)
+						]
 					]
 				]
 			]
