SEnumCombo now only dismisses its own menu when an option is selected.

This was causing issues when an SEnumCombo was embedded into a menu builder, because the top-level menu was keeping the SEnumCombo alive, causing the OnEnumSelectionChangedDelegate to never fire.

Reported from UDN: https://epicgames.lightning.force.com/lightning/r/Case/5004z00001eEsHSAA0/view

diff --git a/Engine/Source/Editor/EditorWidgets/Private/SEnumCombo.cpp b/Engine/Source/Editor/EditorWidgets/Private/SEnumCombo.cpp
--- a/Engine/Source/Editor/EditorWidgets/Private/SEnumCombo.cpp
+++ b/Engine/Source/Editor/EditorWidgets/Private/SEnumCombo.cpp
@@ -162,7 +162,7 @@ FText SEnumComboBox::GetCurrentValueTooltip() const
 TSharedRef<SWidget> SEnumComboBox::OnGetMenuContent()
 {
 	const bool bCloseAfterSelection = !bIsBitflagsEnum;
-	FMenuBuilder MenuBuilder(bCloseAfterSelection, nullptr);
+	FMenuBuilder MenuBuilder(bCloseAfterSelection, nullptr, nullptr, true);
 
 	for (const FEnumInfo& FlagInfo : VisibleEnums)
 	{
