FilterBar: Fix deprecation warnings when including AssetFilter.h

diff --git a/Engine/Source/Editor/ContentBrowser/Private/SFilterList.cpp b/Engine/Source/Editor/ContentBrowser/Private/SFilterList.cpp
--- a/Engine/Source/Editor/ContentBrowser/Private/SFilterList.cpp
+++ b/Engine/Source/Editor/ContentBrowser/Private/SFilterList.cpp
@@ -31,6 +31,7 @@
 #include "ContentBrowserMenuContexts.h"
 #include "Widgets/Images/SImage.h"
 #include "Filters/GenericFilter.h"
+#include "AssetRegistry/ARFilter.h"
 
 #define LOCTEXT_NAMESPACE "ContentBrowser"
 
diff --git a/Engine/Source/Editor/EditorWidgets/Public/Filters/AssetFilter.h b/Engine/Source/Editor/EditorWidgets/Public/Filters/AssetFilter.h
--- a/Engine/Source/Editor/EditorWidgets/Public/Filters/AssetFilter.h
+++ b/Engine/Source/Editor/EditorWidgets/Public/Filters/AssetFilter.h
@@ -8,7 +8,7 @@
 #include "Core/Public/Math/Color.h"
 #include "Core/Public/UObject/NameTypes.h"
 
-#include "ARFilter.h"
+#include "AssetRegistry/ARFilter.h"
 #include "AssetRegistry/IAssetRegistry.h"
 
 /* A filter used by SFilterBar that wraps all the back-end filters into one FFilterBase for convenience */
diff --git a/Engine/Source/Editor/EditorWidgets/Public/Filters/CustomClassFilterData.h b/Engine/Source/Editor/EditorWidgets/Public/Filters/CustomClassFilterData.h
--- a/Engine/Source/Editor/EditorWidgets/Public/Filters/CustomClassFilterData.h
+++ b/Engine/Source/Editor/EditorWidgets/Public/Filters/CustomClassFilterData.h
@@ -2,7 +2,7 @@
 #pragma once
 
 #include "Math/Color.h"
-#include "ARFilter.h"
+#include "AssetRegistry/ARFilter.h"
 
 class IAssetTypeActions;
 class UClass;
diff --git a/Engine/Source/Editor/EditorWidgets/Public/Filters/SAssetFilterBar.h b/Engine/Source/Editor/EditorWidgets/Public/Filters/SAssetFilterBar.h
--- a/Engine/Source/Editor/EditorWidgets/Public/Filters/SAssetFilterBar.h
+++ b/Engine/Source/Editor/EditorWidgets/Public/Filters/SAssetFilterBar.h
@@ -7,7 +7,7 @@
 
 #include "IAssetTypeActions.h"
 #include "AssetTypeCategories.h"
-#include "ARFilter.h"
+#include "AssetRegistry/ARFilter.h"
 #include "AssetToolsModule.h"
 #include "UObject/Object.h"
 #include "SlateGlobals.h"
