FilterBar: Redo icon related changes squashed during merge with SFilterList

diff --git a/Engine/Source/Developer/ToolWidgets/Public/Filters/SBasicFilterBar.h b/Engine/Source/Developer/ToolWidgets/Public/Filters/SBasicFilterBar.h
--- a/Engine/Source/Developer/ToolWidgets/Public/Filters/SBasicFilterBar.h
+++ b/Engine/Source/Developer/ToolWidgets/Public/Filters/SBasicFilterBar.h
@@ -1044,9 +1044,11 @@ protected:
 			"ResetFilters",
 			LOCTEXT("FilterListResetFilters", "Reset Filters"),
 			LOCTEXT("FilterListResetToolTip", "Resets current filter selection"),
-			FSlateIcon(),
-			FUIAction(FExecuteAction::CreateSP(this, &SBasicFilterBar<FilterType>::OnResetFilters))
-			);
+			FSlateIcon(FAppStyle::Get().GetStyleSetName(), "PropertyWindow.DiffersFromDefault"),
+			FUIAction(
+				FExecuteAction::CreateSP(this, &SBasicFilterBar<FilterType>::OnResetFilters),
+				FCanExecuteAction::CreateLambda([this]() { return HasAnyFilters(); }))
+		);
 	}
 
 	/** Helper function to add all custom filters to the Add Filter Menu */
diff --git a/Engine/Source/Editor/EditorWidgets/Public/Filters/SAssetFilterBar.h b/Engine/Source/Editor/EditorWidgets/Public/Filters/SAssetFilterBar.h
--- a/Engine/Source/Editor/EditorWidgets/Public/Filters/SAssetFilterBar.h
+++ b/Engine/Source/Editor/EditorWidgets/Public/Filters/SAssetFilterBar.h
@@ -12,6 +12,7 @@
 #include "UObject/Object.h"
 #include "SlateGlobals.h"
 #include "Logging/LogMacros.h"
+#include "Styling/SlateIconFinder.h"
 
 #include "SAssetFilterBar.generated.h"
 
@@ -513,7 +514,7 @@ protected:
 				NAME_None,
 				LabelText,
 				FText::Format( LOCTEXT("FilterByTooltipPrefix", "Filter by {0}"), LabelText ),
-				FSlateIcon(),
+				FSlateIconFinder::FindIconForClass(CustomClassFilterData->GetClass()),
 				FUIAction(
 					FExecuteAction::CreateSP( const_cast< SAssetFilterBar<FilterType>* >(this), &SAssetFilterBar<FilterType>::FilterByTypeClicked, CustomClassFilterData ),
 					FCanExecuteAction(),
@@ -661,7 +662,7 @@ protected:
 						NAME_None,
 						MenuExpansion->Title,
 						MenuExpansion->Tooltip,
-						FSlateIcon(),
+						FSlateIcon(FAppStyle::Get().GetStyleSetName(), "PlacementBrowser.Icons.Basic"),
 						FUIAction(
 						FExecuteAction::CreateSP( this, &SAssetFilterBar<FilterType>::FilterByTypeCategoryClicked, MenuExpansion, ExpandedCategory->Classes ),
 						FCanExecuteAction(),
