SBasicFilterBar: Add an argument to extend the Add Filter Menu
SAssetFilterBar: Add functions to access and modify Asset Type Filters

diff --git a/Engine/Source/Developer/ToolWidgets/Public/Filters/SBasicFilterBar.h b/Engine/Source/Developer/ToolWidgets/Public/Filters/SBasicFilterBar.h
--- a/Engine/Source/Developer/ToolWidgets/Public/Filters/SBasicFilterBar.h
+++ b/Engine/Source/Developer/ToolWidgets/Public/Filters/SBasicFilterBar.h
@@ -57,11 +57,17 @@ public:
 	/** Delegate for when filters have changed */
 	DECLARE_DELEGATE( FOnFilterChanged );
 	
+	/** Delegate for when filters have changed */
+	DECLARE_DELEGATE_OneParam( FOnExtendAddFilterMenu, UToolMenu* );
+	
  	SLATE_BEGIN_ARGS( SBasicFilterBar<FilterType> ){}
 
  		/** Delegate for when filters have changed */
  		SLATE_EVENT( FOnFilterChanged, OnFilterChanged )
 
+		/** Delegate to extend the Add Filter dropdown */
+		SLATE_EVENT( FOnExtendAddFilterMenu, OnExtendAddFilterMenu )
+
 		/** Initial List of Custom Filters that will be added to the AddFilter Menu */
 		SLATE_ARGUMENT( TArray<TSharedRef<FFilterBase<FilterType>>>, CustomFilters)
 	
@@ -70,6 +76,7 @@ public:
  	void Construct( const FArguments& InArgs )
  	{
 		OnFilterChanged = InArgs._OnFilterChanged;
+ 		OnExtendAddFilterMenu = InArgs._OnExtendAddFilterMenu;
 
  		// A subclass could be using an external filter collection (SFilterList in ContentBrowser)
  		if(!ActiveFilters)
@@ -665,7 +672,7 @@ public:
  		// else -> Already in the desired 'check' state.
  	}
 
- 	/** Returns the check box state of the specified filter (in the filter drop down). This tells whether the filter is pinned or not on the filter bar, but not if filter is active or not. @see IsFrontendFilterActive(). */
+ 	/** Returns the check box state of the specified filter (in the filter drop down). This tells whether the filter is pinned or not on the filter bar, but not if filter is active or not. @see IsFilterActive(). */
  	ECheckBoxState GetFilterCheckState(const TSharedPtr<FFilterBase<FilterType>>& InFilter) const
  	{
  		return InFilter && IsFrontendFilterInUse(InFilter.ToSharedRef()) ? ECheckBoxState::Checked : ECheckBoxState::Unchecked;
@@ -1049,6 +1056,8 @@ protected:
 				FExecuteAction::CreateSP(this, &SBasicFilterBar<FilterType>::OnResetFilters),
 				FCanExecuteAction::CreateLambda([this]() { return HasAnyFilters(); }))
 		);
+		
+		OnExtendAddFilterMenu.ExecuteIfBound(Menu);
 	}
 
 	/** Helper function to add all custom filters to the Add Filter Menu */
@@ -1131,6 +1140,9 @@ protected:
  	/** Delegate for when filters have changed */
  	FOnFilterChanged OnFilterChanged;
 
+	/** Delegate to extend the AddFilter Menu */
+	FOnExtendAddFilterMenu OnExtendAddFilterMenu;
+
 	friend struct FFrontendFilterExternalActivationHelper<FilterType>;
 };
 
diff --git a/Engine/Source/Editor/EditorWidgets/Public/Filters/SAssetFilterBar.h b/Engine/Source/Editor/EditorWidgets/Public/Filters/SAssetFilterBar.h
--- a/Engine/Source/Editor/EditorWidgets/Public/Filters/SAssetFilterBar.h
+++ b/Engine/Source/Editor/EditorWidgets/Public/Filters/SAssetFilterBar.h
@@ -77,7 +77,8 @@ class SAssetFilterBar : public SBasicFilterBar<FilterType>
 public:
 	
 	using FOnFilterChanged = typename SBasicFilterBar<FilterType>::FOnFilterChanged;
-	
+	using FOnExtendAddFilterMenu = typename SBasicFilterBar<FilterType>::FOnExtendAddFilterMenu;
+
 	SLATE_BEGIN_ARGS( SAssetFilterBar<FilterType> )
 		: _UseDefaultAssetFilters(true)
 		{
@@ -87,9 +88,13 @@ public:
 		/** Delegate for when filters have changed */
 		SLATE_EVENT( FOnFilterChanged, OnFilterChanged )
 	
+		/** Delegate to extend the Add Filter dropdown */
+		SLATE_EVENT( FOnExtendAddFilterMenu, OnExtendAddFilterMenu )
+	
 		/** Initial List of Custom Filters that will be added to the AddFilter Menu */
 		SLATE_ARGUMENT( TArray<TSharedRef<FFilterBase<FilterType>>>, CustomFilters)
 
+		/** Initial List of Custom Class filters that will be added to the AddFilter Menu */
 		SLATE_ARGUMENT( TArray<TSharedRef<FCustomClassFilterData>>, CustomClassFilters)
 
 		/** Whether the filter bar should provide the default asset filters */
@@ -106,6 +111,7 @@ public:
 		typename SBasicFilterBar<FilterType>::FArguments Args;
 		Args._OnFilterChanged = InArgs._OnFilterChanged;
 		Args._CustomFilters = InArgs._CustomFilters;
+		Args._OnExtendAddFilterMenu = InArgs._OnExtendAddFilterMenu;
 		
 		SBasicFilterBar<FilterType>::Construct(Args);
 		
@@ -285,12 +291,14 @@ public:
 		return CombinedFilter;
 	}
 
-	/** Check if there is a filter associated with the input class in the filter bar */
-	bool IsAssetTypeFilter(UClass *InClass)
+	/** Check if there is a filter associated with the given class name in the filter bar
+	 * @param	InClassName		The name of the Class the filter is associated with
+	 */
+	bool DoesAssetTypeFilterExist(const FName& InClassName)
 	{
 		for(const TSharedRef<FCustomClassFilterData>& CustomClassFilterData : CustomClassFilters)
 		{
-			if(CustomClassFilterData->GetClass() == InClass)
+			if(CustomClassFilterData->GetFilterName() == InClassName.ToString())
 			{
 				return true;
 			}
@@ -299,17 +307,20 @@ public:
 		return false;
 	}
 
-	/** Append the filter associated with the input class to the filter bar.
-	 *  A filter associated with the input class must exist in the widget for this to succeed.
+	/** Set the check box state of the specified filter (in the filter drop down) and pin/unpin a filter widget on/from the filter bar. When a filter is pinned (was not already pinned), it is activated if requested and deactivated when unpinned.
+	 * @param	InClassName		The name of the Class the filter is associated with (must exist in the widget)
+	 * @param	InCheckState	The CheckState to apply to the flter
+	 * @param	bEnableFilter	Whether the filter should be enabled when it is pinned
 	 */
-	void AddClassFilterToBar(UClass *Class, bool bEnableFilter = true)
+	void SetAssetTypeFilterCheckState(const FName& InClassName, ECheckBoxState InCheckState, bool bEnableFilter = true)
 	{
 		for(const TSharedRef<FCustomClassFilterData>& CustomClassFilterData : CustomClassFilters)
 		{
-			if(CustomClassFilterData->GetClass() == Class)
+			if(CustomClassFilterData->GetFilterName() == InClassName.ToString())
 			{
-				// Make sure it already isn't added
-				if (!IsClassTypeInUse(CustomClassFilterData))
+				bool FilterChecked = IsClassTypeInUse(CustomClassFilterData);
+
+				if (InCheckState == ECheckBoxState::Checked && !FilterChecked)
 				{
 					TSharedRef<SFilter> NewFilter = AddAssetFilterToBar(CustomClassFilterData);
 
@@ -318,6 +329,63 @@ public:
 						NewFilter->SetEnabled(true);
 					}
 				}
+				else if (InCheckState == ECheckBoxState::Unchecked && FilterChecked)
+				{
+					RemoveAssetFilter(CustomClassFilterData); // Unpin the filter widget and deactivate the filter.
+				}
+				// else -> Already in the desired 'check' state.
+				
+			}
+		}
+	}
+
+	/** Returns the check box state of the specified filter (in the filter drop down). This tells whether the filter is pinned or not on the filter bar, but not if filter is active or not.
+	 * @see		IsFilterActive().
+	 * @param	InClassName		The name of the Class the filter is associated with
+	 * @return The CheckState if a filter associated with the class name is in the filter bar, ECheckBoxState::Undetermined otherwise
+	 */
+	
+	ECheckBoxState GetAssetTypeFilterCheckState(const FName& InClassName) const
+	{
+		for(const TSharedRef<FCustomClassFilterData>& CustomClassFilterData : CustomClassFilters)
+		{
+			if(CustomClassFilterData->GetFilterName() == InClassName.ToString())
+			{
+				return IsClassTypeInUse(CustomClassFilterData) ? ECheckBoxState::Checked : ECheckBoxState::Unchecked; 
+			}
+		}
+
+		return ECheckBoxState::Undetermined;
+	}
+
+	/** Returns true if the specified filter is both checked (pinned on the filter bar) and active (contributing to filter the result).
+	 *  @param	InClassName		The name of the Class the filter is associated with
+	 */
+	bool IsAssetTypeFilterActive(const FName& InClassName) const
+	{
+		for(const TSharedPtr<SAssetFilter> AssetFilter : AssetFilters)
+		{
+			if(AssetFilter->GetCustomClassFilterData()->GetFilterName() == InClassName.ToString())
+			{
+				return AssetFilter->IsEnabled();
+			}
+		}
+		
+		return false;
+	}
+
+	/** If a filter with the input class name is Checked (i.e visible in the bar), enable/disable it
+	 * @see SetFilterCheckState to set the check state and GetFilterCheckState to check if it is checked
+	 * @param	InClassName		The name of the Class the filter is associated with
+	 * @param	bEnable			Whether to enable or disable the filter
+	 */
+	void ToggleAssetTypeFilterEnabled(const FName& InClassName, bool bEnable)
+	{
+		for(const TSharedPtr<SAssetFilter> AssetFilter : AssetFilters)
+		{
+			if(AssetFilter->GetCustomClassFilterData()->GetFilterName() == InClassName.ToString())
+			{
+				AssetFilter->SetEnabled(bEnable);
 			}
 		}
 	}
diff --git a/Engine/Source/Editor/EditorWidgets/Public/Filters/SFilterBar.h b/Engine/Source/Editor/EditorWidgets/Public/Filters/SFilterBar.h
--- a/Engine/Source/Editor/EditorWidgets/Public/Filters/SFilterBar.h
+++ b/Engine/Source/Editor/EditorWidgets/Public/Filters/SFilterBar.h
@@ -41,7 +41,8 @@ public:
 	using FOnFilterChanged = typename SBasicFilterBar<FilterType>::FOnFilterChanged;
 	using FConvertItemToAssetData = typename FAssetFilter<FilterType>::FConvertItemToAssetData;
 	using FCompareItemWithClassNames = typename FAssetFilter<FilterType>::FCompareItemWithClassNames;
-
+	using FOnExtendAddFilterMenu = typename SBasicFilterBar<FilterType>::FOnExtendAddFilterMenu;
+	
 	SLATE_BEGIN_ARGS( SFilterBar<FilterType> )
 		: _UseDefaultAssetFilters(true)
 		{
@@ -59,6 +60,9 @@ public:
 
 		/** Initial List of Custom Filters that will be added to the AddFilter Menu */
 		SLATE_ARGUMENT( TArray<TSharedRef<FFilterBase<FilterType>>>, CustomFilters)
+	
+		/** Delegate to extend the Add Filter dropdown */
+		SLATE_EVENT( FOnExtendAddFilterMenu, OnExtendAddFilterMenu )
 
 		/** Initial List of Custom Class Filters that will be added to the AddFilter Menu */
 		SLATE_ARGUMENT( TArray<TSharedRef<FCustomClassFilterData>>, CustomClassFilters)
@@ -76,6 +80,7 @@ public:
 		Args._CustomFilters = InArgs._CustomFilters;
 		Args._UseDefaultAssetFilters = InArgs._UseDefaultAssetFilters;
 		Args._CustomClassFilters = InArgs._CustomClassFilters;
+		Args._OnExtendAddFilterMenu = InArgs._OnExtendAddFilterMenu;
 		
 		SAssetFilterBar<FilterType>::Construct(Args);
 
