TemplateString + customization

diff --git a/Engine/Source/Editor/DetailCustomizations/Private/DetailCustomizations.cpp b/Engine/Source/Editor/DetailCustomizations/Private/DetailCustomizations.cpp
--- a/Engine/Source/Editor/DetailCustomizations/Private/DetailCustomizations.cpp
+++ b/Engine/Source/Editor/DetailCustomizations/Private/DetailCustomizations.cpp
@@ -140,6 +140,7 @@
 #include "StaticMeshComponentDetails.h"
 #include "SubmixDetailsCustomization.h"
 #include "SupportedRangeTypes.h"	// StructsSupportingRangeVisibility
+#include "TemplateStringStructCustomization.h"
 #include "Templates/SharedPointer.h"
 #include "TextCustomization.h"
 #include "TimecodeDetailsCustomization.h"
@@ -310,6 +311,7 @@ void FDetailCustomizationsModule::RegisterPropertyTypeCustomizations()
 	RegisterCustomPropertyTypeLayout("EMaterialShadingModel", FOnGetPropertyTypeCustomizationInstance::CreateStatic(&FMaterialShadingModelCustomization::MakeInstance));
 	RegisterCustomPropertyTypeLayout("DebugCameraControllerSettingsViewModeIndex", FOnGetPropertyTypeCustomizationInstance::CreateStatic(&FDebugCameraControllerSettingsViewModeIndexCustomization::MakeInstance));
 	RegisterCustomPropertyTypeLayout("CustomPrimitiveData", FOnGetPropertyTypeCustomizationInstance::CreateStatic(&FCustomPrimitiveDataCustomization::MakeInstance));
+	RegisterCustomPropertyTypeLayout("TemplateString", FOnGetPropertyTypeCustomizationInstance::CreateStatic(&FTemplateStringStructCustomization::MakeInstance));
 }
 
 #undef REGISTER_UIMINMAX_CUSTOMIZATION
diff --git a/Engine/Source/Editor/DetailCustomizations/Private/TemplateStringStructCustomization.cpp b/Engine/Source/Editor/DetailCustomizations/Private/TemplateStringStructCustomization.cpp
new file mode 100644
--- /dev/null
+++ b/Engine/Source/Editor/DetailCustomizations/Private/TemplateStringStructCustomization.cpp
@@ -0,0 +1,118 @@
+// Copyright Epic Games, Inc. All Rights Reserved.
+
+#include "TemplateStringStructCustomization.h"
+
+#include "DetailWidgetRow.h"
+#include "PropertyHandle.h"
+#include "STemplateStringEditableTextBox.h"
+#include "UObject/TemplateString.h"
+#include "Widgets/DeclarativeSyntaxSupport.h"
+
+#define LOCTEXT_NAMESPACE "TemplateStringStructCustomization"
+
+void FTemplateStringStructCustomization::CustomizeChildren(TSharedRef<IPropertyHandle> StructPropertyHandle, class IDetailChildrenBuilder& StructBuilder, IPropertyTypeCustomizationUtils& StructCustomizationUtils)
+{
+}
+
+void FTemplateStringStructCustomization::CustomizeHeader(TSharedRef<IPropertyHandle> StructPropertyHandle, class FDetailWidgetRow& HeaderRow, IPropertyTypeCustomizationUtils& StructCustomizationUtils)
+{
+	TemplateStringProperty = StructPropertyHandle->GetChildHandle(GET_MEMBER_NAME_CHECKED(FTemplateString, Template));
+	check(TemplateStringProperty);
+
+	HeaderRow
+	.NameContent()
+	[
+		StructPropertyHandle->CreatePropertyNameWidget()
+	]
+	.ValueContent()
+	.MaxDesiredWidth(0.0f)
+	.MinDesiredWidth(125.0f)
+	[
+		SNew(STemplateStringEditableTextBox)
+		.ToolTipText_Lambda([ToolTipText = GetToolTip()]()
+		{
+			return ToolTipText;
+		})
+		.Text_Raw(this, &FTemplateStringStructCustomization::GetText)
+		.OnTextChanged(this, &FTemplateStringStructCustomization::SetText)
+	];
+}
+
+FText FTemplateStringStructCustomization::GetText() const
+{
+	if(TemplateStringProperty.IsValid())
+	{
+		FString TemplateString;
+		TemplateStringProperty->GetValue(TemplateString);
+
+		return FText::FromString(TemplateString);
+	}
+
+	return FText::GetEmpty();
+}
+
+void FTemplateStringStructCustomization::SetText(const FText& InNewText) const
+{
+	if(TemplateStringProperty.IsValid())
+	{
+		TemplateStringProperty->SetValue(InNewText.ToString());
+	}
+}
+
+FText FTemplateStringStructCustomization::GetToolTip() const
+{
+	static const FText EmptyText = FText::GetEmpty();
+
+	if(TemplateStringProperty.IsValid())
+	{
+		if(CachedTooltip.IsEmpty())
+		{
+			FTextBuilder TextBuilder;
+			if(TemplateStringProperty->GetParentHandle().IsValid())
+			{
+				TextBuilder.AppendLine(TemplateStringProperty->GetParentHandle()->GetToolTipText());
+			}
+
+			const TArray<FString>& ValidArgs = GetValidArguments();
+			if(!ValidArgs.IsEmpty())
+			{
+				TextBuilder.AppendLine(LOCTEXT("ValidArgs_ToolTipHeading", "Valid Arguments:"));
+				TextBuilder.Indent();
+
+				for(const FString& Arg : ValidArgs)
+				{
+					TextBuilder.AppendLine(Arg);
+				}
+			}
+
+			CachedTooltip = TextBuilder.ToText();
+		}
+		
+		return CachedTooltip;
+	}
+
+	return EmptyText;
+}
+
+const TArray<FString>& FTemplateStringStructCustomization::GetValidArguments() const
+{
+	static TArray<FString> EmptyArray;
+
+	if(TemplateStringProperty.IsValid())
+	{
+		if(ValidArguments.IsEmpty())
+		{
+			if(TemplateStringProperty->GetParentHandle()->HasMetaData(TEXT("ValidArgs")))
+			{
+				const FString ValidArgStr = TemplateStringProperty->GetParentHandle()->GetMetaData(TEXT("ValidArgs"));
+				ValidArgStr.ParseIntoArrayWS(ValidArguments, TEXT(","), true);
+			}
+		}
+
+		return ValidArguments;
+	}
+
+	return EmptyArray;
+}
+
+#undef LOCTEXT_NAMESPACE
diff --git a/Engine/Source/Editor/DetailCustomizations/Private/TemplateStringStructCustomization.h b/Engine/Source/Editor/DetailCustomizations/Private/TemplateStringStructCustomization.h
new file mode 100644
--- /dev/null
+++ b/Engine/Source/Editor/DetailCustomizations/Private/TemplateStringStructCustomization.h
@@ -0,0 +1,54 @@
+// Copyright Epic Games, Inc. All Rights Reserved.
+
+#pragma once
+
+#include "CoreMinimal.h"
+#include "IPropertyTypeCustomization.h"
+#include "Core/Private/Internationalization/TextCache.h"
+
+class IPropertyHandle;
+
+/**
+ * Implements a details view customization for the FTemplateString structure.
+ */
+class FTemplateStringStructCustomization
+	: public IPropertyTypeCustomization
+{
+public:
+	/**
+	 * Creates an instance of this class.
+	 *
+	 * @return The new instance.
+	 */
+	static TSharedRef<IPropertyTypeCustomization> MakeInstance( )
+	{
+		return MakeShareable(new FTemplateStringStructCustomization());
+	}
+
+public:
+	virtual void CustomizeChildren( TSharedRef<class IPropertyHandle> StructPropertyHandle, class IDetailChildrenBuilder& StructBuilder, IPropertyTypeCustomizationUtils& StructCustomizationUtils ) override;
+	virtual void CustomizeHeader( TSharedRef<class IPropertyHandle> StructPropertyHandle, class FDetailWidgetRow& HeaderRow, IPropertyTypeCustomizationUtils& StructCustomizationUtils ) override;
+
+private:
+	/** Get the property string value (as FText). */
+	FText GetText() const;
+
+	/** Set the property string value (from FText). */
+	void SetText(const FText& InNewText) const;
+
+	/** Get/Create tooltip, with list of valid arguments. */	
+	FText GetToolTip() const;
+
+	/** Get the list of valid arguments. */
+	const TArray<FString>& GetValidArguments() const;
+
+private:
+	/** Pointer to the string that will be set when changing the path */
+	TSharedPtr<IPropertyHandle> TemplateStringProperty;
+
+	/** Cached tooltip, created by GetToolTip(). */
+	mutable FText CachedTooltip;
+
+	/** Store the valid arguments for this property. */
+	mutable TArray<FString> ValidArguments;
+};
diff --git a/Engine/Source/Editor/EditorWidgets/Private/EditorWidgetsModule.cpp b/Engine/Source/Editor/EditorWidgets/Private/EditorWidgetsModule.cpp
--- a/Engine/Source/Editor/EditorWidgets/Private/EditorWidgetsModule.cpp
+++ b/Engine/Source/Editor/EditorWidgets/Private/EditorWidgetsModule.cpp
@@ -1,13 +1,15 @@
 // Copyright Epic Games, Inc. All Rights Reserved.
 
 #include "EditorWidgetsModule.h"
-#include "Modules/ModuleManager.h"
-#include "Widgets/DeclarativeSyntaxSupport.h"
-#include "SObjectNameEditableTextBox.h"
+
+#include "EditorWidgetsStyle.h"
+#include "ITransportControl.h"
 #include "ObjectNameEditSinkRegistry.h"
 #include "SAssetDiscoveryIndicator.h"
-#include "ITransportControl.h"
+#include "SObjectNameEditableTextBox.h"
 #include "STransportControl.h"
+#include "Modules/ModuleManager.h"
+#include "Widgets/DeclarativeSyntaxSupport.h"
 
 IMPLEMENT_MODULE( FEditorWidgetsModule, EditorWidgets );
 
@@ -15,12 +17,16 @@ const FName FEditorWidgetsModule::EditorWidgetsAppIdentifier( TEXT( "EditorWidge
 
 void FEditorWidgetsModule::StartupModule()
 {
+	FEditorWidgetsStyle::Initialize();
+	
 	ObjectNameEditSinkRegistry = MakeShared<UE::EditorWidgets::FObjectNameEditSinkRegistry>();
 }
 
 void FEditorWidgetsModule::ShutdownModule()
 {
 	ObjectNameEditSinkRegistry.Reset();
+
+	FEditorWidgetsStyle::Shutdown();
 }
 
 TSharedRef<IObjectNameEditableTextBox> FEditorWidgetsModule::CreateObjectNameEditableTextBox(const TArray<TWeakObjectPtr<UObject>>& Objects)
diff --git a/Engine/Source/Editor/EditorWidgets/Private/EditorWidgetsStyle.cpp b/Engine/Source/Editor/EditorWidgets/Private/EditorWidgetsStyle.cpp
new file mode 100644
--- /dev/null
+++ b/Engine/Source/Editor/EditorWidgets/Private/EditorWidgetsStyle.cpp
@@ -0,0 +1,67 @@
+﻿// Copyright Epic Games, Inc. All Rights Reserved.
+
+#include "EditorWidgetsStyle.h"
+#include "Styling/AppStyle.h"
+#include "Styling/SlateStyleRegistry.h"
+#include "Styling/SlateTypes.h"
+
+TSharedPtr<FSlateStyleSet> FEditorWidgetsStyle::StyleSet = nullptr;
+
+#define IMAGE_BRUSH( RelativePath, ... ) FSlateImageBrush( StyleSet->RootToContentDir( RelativePath, TEXT(".png") ), __VA_ARGS__ )
+#define BOX_BRUSH( RelativePath, ... ) FSlateBoxBrush( StyleSet->RootToContentDir( RelativePath, TEXT(".png") ), __VA_ARGS__ )
+#define BORDER_BRUSH( RelativePath, ... ) FSlateBorderBrush( StyleSet->RootToContentDir( RelativePath, TEXT(".png") ), __VA_ARGS__ )
+#define DEFAULT_FONT(...) FCoreStyle::GetDefaultFontStyle(__VA_ARGS__)
+
+void FEditorWidgetsStyle::Initialize()
+{
+	// Only register once
+	if(StyleSet.IsValid())
+	{
+		return;
+	}
+
+	StyleSet = MakeShareable(new FSlateStyleSet("EditorWidgets"));
+
+	const FEditableTextBoxStyle& NormalEditableTextBoxStyle = FAppStyle::GetWidgetStyle<FEditableTextBoxStyle>("NormalEditableTextBox");
+
+	const FTextBlockStyle NormalText = FTextBlockStyle(FAppStyle::GetWidgetStyle<FTextBlockStyle>("TextEditor.NormalText"))
+		.SetColorAndOpacity(NormalEditableTextBoxStyle.ForegroundColor)
+		.SetHighlightColor(NormalEditableTextBoxStyle.FocusedForegroundColor)
+		.SetFont(NormalEditableTextBoxStyle.Font)
+		.SetFontSize(NormalEditableTextBoxStyle.Font.Size);
+
+	// Text editor
+	{
+		StyleSet->Set("TextEditor.NormalText", NormalText);
+
+		StyleSet->Set("SyntaxHighlight.Template.Normal", NormalText);
+		StyleSet->Set("SyntaxHighlight.Template.Argument", FAppStyle::GetWidgetStyle<FTextBlockStyle>("RichTextBlock.BoldHighlight"));
+	}
+
+	FSlateStyleRegistry::RegisterSlateStyle( *StyleSet.Get() );
+}
+
+#undef IMAGE_BRUSH
+#undef BOX_BRUSH
+#undef BORDER_BRUSH
+#undef DEFAULT_FONT
+
+void FEditorWidgetsStyle::Shutdown()
+{
+	if(StyleSet.IsValid())
+	{
+		FSlateStyleRegistry::UnRegisterSlateStyle(*StyleSet.Get());
+		ensure(StyleSet.IsUnique());
+		StyleSet.Reset();
+	}
+}
+
+const ISlateStyle& FEditorWidgetsStyle::Get()
+{
+	return *(StyleSet.Get());
+}
+
+const FName& FEditorWidgetsStyle::GetStyleSetName()
+{
+	return StyleSet->GetStyleSetName();
+}
diff --git a/Engine/Source/Editor/EditorWidgets/Private/EditorWidgetsStyle.h b/Engine/Source/Editor/EditorWidgets/Private/EditorWidgetsStyle.h
new file mode 100644
--- /dev/null
+++ b/Engine/Source/Editor/EditorWidgets/Private/EditorWidgetsStyle.h
@@ -0,0 +1,21 @@
+﻿// Copyright Epic Games, Inc. All Rights Reserved.
+
+#pragma once
+
+#include "CoreMinimal.h"
+#include "Styling/SlateStyle.h"
+
+class FEditorWidgetsStyle
+{
+public:
+	static void Initialize();
+	static void Shutdown();
+
+	static const ISlateStyle& Get();
+
+	static const FName& GetStyleSetName();
+
+private:
+	/** Singleton instances of this style. */
+	static TSharedPtr<class FSlateStyleSet> StyleSet;	
+};
diff --git a/Engine/Source/Editor/EditorWidgets/Private/STemplateStringEditableTextBox.cpp b/Engine/Source/Editor/EditorWidgets/Private/STemplateStringEditableTextBox.cpp
new file mode 100644
--- /dev/null
+++ b/Engine/Source/Editor/EditorWidgets/Private/STemplateStringEditableTextBox.cpp
@@ -0,0 +1,24 @@
+﻿// Copyright Epic Games, Inc. All Rights Reserved.
+
+#include "STemplateStringEditableTextBox.h"
+
+#include "EditorWidgetsStyle.h"
+#include "SlateOptMacros.h"
+#include "TemplateStringSyntaxHighlighterMarshaller.h"
+
+BEGIN_SLATE_FUNCTION_BUILD_OPTIMIZATION
+
+void STemplateStringEditableTextBox::Construct(const FArguments& InArgs)
+{
+	SMultiLineEditableTextBox::Construct(SMultiLineEditableTextBox::FArguments()
+		.TextStyle(&FEditorWidgetsStyle::Get().GetWidgetStyle<FTextBlockStyle>("TextEditor.NormalText"))
+		.Text(InArgs._Text)
+		.Marshaller(FTemplateStringSyntaxHighlighterMarshaller::Create(FTemplateStringSyntaxHighlighterMarshaller::FSyntaxTextStyle()))
+		.AllowMultiLine(false)
+		.OverflowPolicy(ETextOverflowPolicy::Ellipsis)
+		.Margin(0.0f)
+		.OnTextChanged(InArgs._OnTextChanged)
+	);
+}
+
+END_SLATE_FUNCTION_BUILD_OPTIMIZATION
diff --git a/Engine/Source/Editor/EditorWidgets/Private/TemplateStringSyntaxHighlighterMarshaller.cpp b/Engine/Source/Editor/EditorWidgets/Private/TemplateStringSyntaxHighlighterMarshaller.cpp
new file mode 100644
--- /dev/null
+++ b/Engine/Source/Editor/EditorWidgets/Private/TemplateStringSyntaxHighlighterMarshaller.cpp
@@ -0,0 +1,98 @@
+﻿// Copyright Epic Games, Inc. All Rights Reserved.
+
+#include "TemplateStringSyntaxHighlighterMarshaller.h"
+#include "Framework/Text/IRun.h"
+#include "Framework/Text/TextLayout.h"
+#include "Framework/Text/ISlateRun.h"
+#include "Framework/Text/SlateTextRun.h"
+
+constexpr static const TCHAR* Operators[] = {
+	TEXT("{"),
+	TEXT("}")
+};
+
+TSharedRef<FTemplateStringSyntaxHighlighterMarshaller> FTemplateStringSyntaxHighlighterMarshaller::Create(const FSyntaxTextStyle& InSyntaxTextStyle)
+{
+	TArray<FSyntaxTokenizer::FRule> TokenizerRules;
+
+	// operators
+	for(const TCHAR* Operator : Operators)
+	{
+		TokenizerRules.Emplace(FSyntaxTokenizer::FRule(Operator));
+	}	
+
+	return MakeShared<FTemplateStringSyntaxHighlighterMarshaller>(FSyntaxTokenizer::Create(TokenizerRules), InSyntaxTextStyle);
+}
+
+void FTemplateStringSyntaxHighlighterMarshaller::ParseTokens(const FString& SourceString, FTextLayout& TargetTextLayout, TArray<ISyntaxTokenizer::FTokenizedLine> TokenizedLines)
+{
+	enum class EParseState : uint8
+	{
+		None,
+		LookingForArgument,
+	};
+
+	TArray<FTextLayout::FNewLineData> LinesToAdd;
+	LinesToAdd.Reserve(TokenizedLines.Num());
+
+	// Parse the tokens, generating the styled runs for each line
+	EParseState ParseState = EParseState::None;
+	for(const ISyntaxTokenizer::FTokenizedLine& TokenizedLine : TokenizedLines)
+	{
+		TSharedRef<FString> ModelString = MakeShared<FString>();
+		TArray<TSharedRef<IRun>> Runs;
+
+		for(const ISyntaxTokenizer::FToken& Token : TokenizedLine.Tokens)
+		{
+			const FStringView TokenText(*SourceString + Token.Range.BeginIndex, Token.Range.Len());
+
+			const FTextRange ModelRange(ModelString->Len(), ModelString->Len() + TokenText.Len());
+			ModelString->Append(TokenText);
+
+			FRunInfo RunInfo(TEXT("SyntaxHighlight.Template.Normal"));
+			FTextBlockStyle TextBlockStyle = SyntaxTextStyle.NormalTextStyle;
+
+			bool bHasMatchedSyntax = false;
+			if(Token.Type == ISyntaxTokenizer::ETokenType::Syntax)
+			{
+				if(ParseState == EParseState::None && TokenText == TEXT("{"))
+				{
+					RunInfo.Name = TEXT("SyntaxHighlight.Template.Argument");
+					TextBlockStyle = SyntaxTextStyle.ArgumentTextStyle;
+					ParseState = EParseState::LookingForArgument;
+					bHasMatchedSyntax = true;
+				}
+				else if(ParseState == EParseState::LookingForArgument && TokenText == TEXT("}"))
+				{
+					RunInfo.Name = TEXT("SyntaxHighlight.Template.Normal");
+					TextBlockStyle = SyntaxTextStyle.ArgumentTextStyle; // Closing brace should still use this style
+					ParseState = EParseState::None;
+				}
+			}
+			
+			// It's possible that we fail to match a syntax token if we're in a state where it isn't parsed
+			// In this case, we treat it as a literal token
+			if(Token.Type == ISyntaxTokenizer::ETokenType::Literal || !bHasMatchedSyntax)
+			{
+				if(ParseState == EParseState::LookingForArgument)
+				{
+					RunInfo.Name = TEXT("SyntaxHighlight.Template.Argument");
+					TextBlockStyle = SyntaxTextStyle.ArgumentTextStyle;
+				}
+			}
+
+			TSharedRef< ISlateRun > Run = FSlateTextRun::Create(RunInfo, ModelString, TextBlockStyle, ModelRange);
+			Runs.Add(Run);
+		}
+
+		LinesToAdd.Emplace(MoveTemp(ModelString), MoveTemp(Runs));
+	}
+
+	TargetTextLayout.AddLines(LinesToAdd);
+}
+
+FTemplateStringSyntaxHighlighterMarshaller::FTemplateStringSyntaxHighlighterMarshaller(TSharedPtr< ISyntaxTokenizer > InTokenizer, const FSyntaxTextStyle& InSyntaxTextStyle)
+	: FSyntaxHighlighterTextLayoutMarshaller(MoveTemp(InTokenizer))
+	, SyntaxTextStyle(InSyntaxTextStyle)
+{
+}
diff --git a/Engine/Source/Editor/EditorWidgets/Private/TemplateStringSyntaxHighlighterMarshaller.h b/Engine/Source/Editor/EditorWidgets/Private/TemplateStringSyntaxHighlighterMarshaller.h
new file mode 100644
--- /dev/null
+++ b/Engine/Source/Editor/EditorWidgets/Private/TemplateStringSyntaxHighlighterMarshaller.h
@@ -0,0 +1,51 @@
+﻿// Copyright Epic Games, Inc. All Rights Reserved.
+
+#pragma once
+
+#include "CoreMinimal.h"
+#include "Styling/SlateTypes.h"
+#include "EditorWidgetsStyle.h"
+#include "Framework/Text/SyntaxTokenizer.h"
+#include "Framework/Text/SyntaxHighlighterTextLayoutMarshaller.h"
+
+class FTextLayout;
+
+/**
+ * Get/set the raw text to/from a text layout, and also inject syntax highlighting for our rich-text markup
+ */
+class FTemplateStringSyntaxHighlighterMarshaller : public FSyntaxHighlighterTextLayoutMarshaller
+{
+public:
+	struct FSyntaxTextStyle
+	{
+		FSyntaxTextStyle()
+			: NormalTextStyle(FEditorWidgetsStyle::Get().GetWidgetStyle<FTextBlockStyle>("SyntaxHighlight.Template.Normal"))
+			, ArgumentTextStyle(FEditorWidgetsStyle::Get().GetWidgetStyle<FTextBlockStyle>("SyntaxHighlight.Template.Argument"))
+		{
+		}
+
+		FSyntaxTextStyle(const FTextBlockStyle& InNormalTextStyle, const FTextBlockStyle& InArgumentTextStyle)
+			: NormalTextStyle(InNormalTextStyle)
+			, ArgumentTextStyle(InArgumentTextStyle)
+		{
+		}
+
+		FTextBlockStyle NormalTextStyle;
+		FTextBlockStyle ArgumentTextStyle;
+	};
+
+	static TSharedRef<FTemplateStringSyntaxHighlighterMarshaller> Create(const FSyntaxTextStyle& InSyntaxTextStyle);
+
+	virtual ~FTemplateStringSyntaxHighlighterMarshaller() override = default;
+
+protected:
+	// Allows MakeShared with private constructor
+	friend class SharedPointerInternals::TIntrusiveReferenceController<FTemplateStringSyntaxHighlighterMarshaller, ESPMode::ThreadSafe>;
+	
+	virtual void ParseTokens(const FString& SourceString, FTextLayout& TargetTextLayout, TArray<ISyntaxTokenizer::FTokenizedLine> TokenizedLines) override;
+
+	FTemplateStringSyntaxHighlighterMarshaller(TSharedPtr<ISyntaxTokenizer> InTokenizer, const FSyntaxTextStyle& InSyntaxTextStyle);
+
+	/** Styles used to display the text */
+	FSyntaxTextStyle SyntaxTextStyle;
+};
diff --git a/Engine/Source/Editor/EditorWidgets/Public/STemplateStringEditableTextBox.h b/Engine/Source/Editor/EditorWidgets/Public/STemplateStringEditableTextBox.h
new file mode 100644
--- /dev/null
+++ b/Engine/Source/Editor/EditorWidgets/Public/STemplateStringEditableTextBox.h
@@ -0,0 +1,31 @@
+﻿// Copyright Epic Games, Inc. All Rights Reserved.
+
+#pragma once
+
+#include "CoreMinimal.h"
+#include "Widgets/Input/SMultiLineEditableTextBox.h"
+
+class ITextLayoutMarshaller;
+
+class EDITORWIDGETS_API STemplateStringEditableTextBox : public SMultiLineEditableTextBox
+{
+public:
+	SLATE_BEGIN_ARGS(STemplateStringEditableTextBox)
+	{}
+		/** The initial text that will appear in the widget. */
+		SLATE_ATTRIBUTE(FText, Text)
+
+		/** The list of available arguments to use in this template string. */
+		SLATE_ATTRIBUTE(TArray<FString>, ValidArguments)
+
+		/** The marshaller used to get/set the raw text to/from the text layout. */
+		SLATE_ARGUMENT(TSharedPtr<ITextLayoutMarshaller>, Marshaller)
+
+		/** Called whenever the text is changed interactively by the user. */
+		SLATE_EVENT(FOnTextChanged, OnTextChanged)
+	
+	SLATE_END_ARGS()
+
+	/** Constructs this widget with InArgs */
+	void Construct(const FArguments& InArgs);
+};
diff --git a/Engine/Source/Runtime/CoreUObject/Private/UObject/TemplateString.cpp b/Engine/Source/Runtime/CoreUObject/Private/UObject/TemplateString.cpp
new file mode 100644
--- /dev/null
+++ b/Engine/Source/Runtime/CoreUObject/Private/UObject/TemplateString.cpp
@@ -0,0 +1,51 @@
+﻿// Copyright Epic Games, Inc. All Rights Reserved.
+
+#include "UObject/TemplateString.h"
+
+#include "Algo/Count.h"
+#include "Internationalization/TextFormatter.h"
+
+bool FTemplateString::IsValid(const TArray<FString>& InValidArgs) const
+{
+	constexpr static TCHAR OpenBrace = '{';
+	constexpr static TCHAR CloseBrace = '}';
+	
+	if(InValidArgs.IsEmpty())
+	{
+		int32 OpenBraceCount = 0;
+		int32 CloseBraceCount = 0;
+		for(const TCHAR& Char : Template)
+		{
+			if(Char == OpenBrace)
+			{
+				++OpenBraceCount;
+			}
+			else if(Char == CloseBrace)
+			{
+				++CloseBraceCount;
+			}
+		}
+
+		// If false, mismatched braces!
+		return OpenBraceCount == CloseBraceCount;
+	}
+	else
+	{
+		const TSet<FString> ValidArgSet(InValidArgs);
+
+		TArray<FString> ArgsFound;
+		const FTextFormat TextFormat(FText::FromString(Template));
+		// If false, format string itself is probably malformed
+		if(!TextFormat.IsValid())
+		{
+			return false;
+		}
+		
+		TextFormat.GetFormatArgumentNames(ArgsFound);
+
+		const TSet<FString> ArgsFoundSet(ArgsFound);
+
+		// True if all found arguments exist in ValidArgs
+		return ValidArgSet.Includes(ArgsFoundSet);		
+	}
+}
diff --git a/Engine/Source/Runtime/CoreUObject/Public/UObject/NoExportTypes.h b/Engine/Source/Runtime/CoreUObject/Public/UObject/NoExportTypes.h
--- a/Engine/Source/Runtime/CoreUObject/Public/UObject/NoExportTypes.h
+++ b/Engine/Source/Runtime/CoreUObject/Public/UObject/NoExportTypes.h
@@ -2645,8 +2645,8 @@ struct FFilePath
 	GENERATED_BODY()
 
 	/**
-		* The path to the file.
-		*/
+	* The path to the file.
+	*/
 	UPROPERTY(EditAnywhere, BlueprintReadWrite, Category = FilePath)
 	FString FilePath;
 };
@@ -2658,15 +2658,25 @@ struct FDirectoryPath
 	GENERATED_BODY()
 
 	/**
-		* The path to the directory.
-		*/
+	* The path to the directory.
+	*/
 	UPROPERTY(EditAnywhere, BlueprintReadWrite, Category = Path)
 	FString Path;
 };
 
+/** Structure for templated strings that are displayed in the editor with a allowed args. */
+USTRUCT(noexport, BlueprintType)
+struct FTemplateString
+{
+	GENERATED_BODY()
 
+	/**
+	* The format string.
+	*/
+	UPROPERTY(EditAnywhere, BlueprintReadWrite, Category = Template)
+	FString Template;
+};
 
 /// @endcond
 
 #endif
-
diff --git a/Engine/Source/Runtime/CoreUObject/Public/UObject/TemplateString.h b/Engine/Source/Runtime/CoreUObject/Public/UObject/TemplateString.h
new file mode 100644
--- /dev/null
+++ b/Engine/Source/Runtime/CoreUObject/Public/UObject/TemplateString.h
@@ -0,0 +1,17 @@
+﻿// Copyright Epic Games, Inc. All Rights Reserved.
+
+#pragma once
+
+#include "CoreMinimal.h"
+
+/** Structure for templated strings that are displayed in the editor with a allowed args. */
+struct FTemplateString
+{
+	/**
+	* The template string.
+	*/
+	FString Template;
+
+	/** Returns validity based on brace matching, and if provided, arg presence in ValidArgs. */
+	bool IsValid(const TArray<FString>& InValidArgs = {}) const;
+};
