EditorWidgetsStyle cis fix 176112

diff --git a/Engine/Source/Editor/EditorWidgets/Private/EditorWidgetsStyle.cpp b/Engine/Source/Editor/EditorWidgets/Private/EditorWidgetsStyle.cpp
--- a/Engine/Source/Editor/EditorWidgets/Private/EditorWidgetsStyle.cpp
+++ b/Engine/Source/Editor/EditorWidgets/Private/EditorWidgetsStyle.cpp
@@ -24,7 +24,7 @@ void FEditorWidgetsStyle::Initialize()
 
 	const FEditableTextBoxStyle& NormalEditableTextBoxStyle = FAppStyle::GetWidgetStyle<FEditableTextBoxStyle>("NormalEditableTextBox");
 
-	const FTextBlockStyle NormalText = FTextBlockStyle(FAppStyle::GetWidgetStyle<FTextBlockStyle>("TextEditor.NormalText"))
+	const FTextBlockStyle NormalText = FTextBlockStyle()
 		.SetColorAndOpacity(NormalEditableTextBoxStyle.ForegroundColor)
 		.SetHighlightColor(NormalEditableTextBoxStyle.FocusedForegroundColor)
 		.SetFont(NormalEditableTextBoxStyle.Font)
