SAssetFilterBar: Change all functions using FNames to use FTopLevelAssetPath

diff --git a/Engine/Plugins/Editor/AssetManagerEditor/Source/AssetManagerEditor/Private/ReferenceViewer/SReferenceViewer.cpp b/Engine/Plugins/Editor/AssetManagerEditor/Source/AssetManagerEditor/Private/ReferenceViewer/SReferenceViewer.cpp
--- a/Engine/Plugins/Editor/AssetManagerEditor/Source/AssetManagerEditor/Private/ReferenceViewer/SReferenceViewer.cpp
+++ b/Engine/Plugins/Editor/AssetManagerEditor/Source/AssetManagerEditor/Private/ReferenceViewer/SReferenceViewer.cpp
@@ -735,11 +735,9 @@ void SReferenceViewer::OnUpdateFilterBar()
 			FilterWidget->RemoveAllFilters();
 			for (const FTopLevelAssetPath& AssetClassPath : AllClasses)
 			{
-				// @todo make DoesAssetTypeFilterExist/SetAssetTypeFilterCheckState use FTopLevelAssetPath
-				const FName AssetClass = FName(*AssetClassPath.ToString());
-				if (FilterWidget->DoesAssetTypeFilterExist(AssetClass))
+				if (FilterWidget->DoesAssetTypeFilterExist(AssetClassPath))
 				{
-					FilterWidget->SetAssetTypeFilterCheckState(AssetClass, ECheckBoxState::Checked);
+					FilterWidget->SetAssetTypeFilterCheckState(AssetClassPath, ECheckBoxState::Checked);
 				}
 			}
 
diff --git a/Engine/Plugins/Editor/AssetManagerEditor/Source/AssetManagerEditor/Private/ReferenceViewer/SReferenceViewerFilterBar.cpp b/Engine/Plugins/Editor/AssetManagerEditor/Source/AssetManagerEditor/Private/ReferenceViewer/SReferenceViewerFilterBar.cpp
--- a/Engine/Plugins/Editor/AssetManagerEditor/Source/AssetManagerEditor/Private/ReferenceViewer/SReferenceViewerFilterBar.cpp
+++ b/Engine/Plugins/Editor/AssetManagerEditor/Source/AssetManagerEditor/Private/ReferenceViewer/SReferenceViewerFilterBar.cpp
@@ -14,8 +14,8 @@ void SReferenceViewerFilterBar::SaveSettings(const FString& IniFilename, const F
 		{
 			for (TSharedRef<SAssetFilter> CurrentAssetFilter : AssetFilters)
 			{
-				FName AssetFilterName = FName(CurrentAssetFilter->GetCustomClassFilterData()->GetFilterName());
-				SaveFilters.Add(FilterState(AssetFilterName, CurrentAssetFilter->IsEnabled()));
+				FTopLevelAssetPath AssetFilterPath = CurrentAssetFilter->GetCustomClassFilterData()->GetClassPathName();
+				SaveFilters.Add(FilterState(AssetFilterPath, CurrentAssetFilter->IsEnabled()));
 			}
 			Settings->SetUserFilters(SaveFilters);
 		}
@@ -33,10 +33,10 @@ void SReferenceViewerFilterBar::LoadSettings(const FString& IniFilename, const F
 			RemoveAllFilters();
 			for (FilterState& State : SavedFilters)
 			{
-				if (DoesAssetTypeFilterExist(State.FilterName))
+				if (DoesAssetTypeFilterExist(State.FilterPath))
 				{
-					SetAssetTypeFilterCheckState(State.FilterName, ECheckBoxState::Checked);
-					ToggleAssetTypeFilterEnabled(State.FilterName, State.bIsEnabled);
+					SetAssetTypeFilterCheckState(State.FilterPath, ECheckBoxState::Checked);
+					ToggleAssetTypeFilterEnabled(State.FilterPath, State.bIsEnabled);
 				}
 			}
 		}
diff --git a/Engine/Plugins/Editor/AssetManagerEditor/Source/AssetManagerEditor/Public/ReferenceViewer/ReferenceViewerSettings.h b/Engine/Plugins/Editor/AssetManagerEditor/Source/AssetManagerEditor/Public/ReferenceViewer/ReferenceViewerSettings.h
--- a/Engine/Plugins/Editor/AssetManagerEditor/Source/AssetManagerEditor/Public/ReferenceViewer/ReferenceViewerSettings.h
+++ b/Engine/Plugins/Editor/AssetManagerEditor/Source/AssetManagerEditor/Public/ReferenceViewer/ReferenceViewerSettings.h
@@ -16,17 +16,17 @@ struct FilterState
 	GENERATED_BODY();
 
 	FilterState()
-	: FilterName(NAME_None)
+	: FilterPath()
 	, bIsEnabled(false)
 	{}
 
-	FilterState(FName InFilterName, bool InState)
-	: FilterName(InFilterName)
+	FilterState(FTopLevelAssetPath InFilterPath, bool InState)
+	: FilterPath(InFilterPath)
 	, bIsEnabled(InState)
 	{}
 
 	UPROPERTY()
-	FName FilterName;
+	FTopLevelAssetPath FilterPath;
 
 	UPROPERTY()
 	bool bIsEnabled;
diff --git a/Engine/Source/Editor/EditorWidgets/Private/Filters/CustomClassFilterData.cpp b/Engine/Source/Editor/EditorWidgets/Private/Filters/CustomClassFilterData.cpp
--- a/Engine/Source/Editor/EditorWidgets/Private/Filters/CustomClassFilterData.cpp
+++ b/Engine/Source/Editor/EditorWidgets/Private/Filters/CustomClassFilterData.cpp
@@ -80,3 +80,12 @@ FString FCustomClassFilterData::GetFilterName() const
 	return Class->GetFName().ToString();
 }
 
+FTopLevelAssetPath FCustomClassFilterData::GetClassPathName() const
+{
+	if(TSharedPtr<IAssetTypeActions> AssetTypeActionsPin = AssetTypeActions.Pin())
+	{
+		return AssetTypeActionsPin->GetClassPathName();
+	}
+
+	return Class->GetClassPathName();
+}
diff --git a/Engine/Source/Editor/EditorWidgets/Public/Filters/CustomClassFilterData.h b/Engine/Source/Editor/EditorWidgets/Public/Filters/CustomClassFilterData.h
--- a/Engine/Source/Editor/EditorWidgets/Public/Filters/CustomClassFilterData.h
+++ b/Engine/Source/Editor/EditorWidgets/Public/Filters/CustomClassFilterData.h
@@ -57,6 +57,11 @@ public:
 	/** Get the display name of this filter class */
 	FString GetFilterName() const;
 
+	/**
+	 * Returns class path name as a package + class FName pair
+	 */
+	FTopLevelAssetPath GetClassPathName() const;
+
 private:
 
 	/** The Class the filter is associated with (if it does not have an AssetTypeAction) */
diff --git a/Engine/Source/Editor/EditorWidgets/Public/Filters/SAssetFilterBar.h b/Engine/Source/Editor/EditorWidgets/Public/Filters/SAssetFilterBar.h
--- a/Engine/Source/Editor/EditorWidgets/Public/Filters/SAssetFilterBar.h
+++ b/Engine/Source/Editor/EditorWidgets/Public/Filters/SAssetFilterBar.h
@@ -291,14 +291,14 @@ public:
 		return CombinedFilter;
 	}
 
-	/** Check if there is a filter associated with the given class name in the filter bar
-	 * @param	InClassName		The name of the Class the filter is associated with
+	/** Check if there is a filter associated with the given class represented by FTopLevelAssetPath in the filter bar
+	 * @param	InClassPath		The Class the filter is associated with
 	 */
-	bool DoesAssetTypeFilterExist(const FName& InClassName)
+	bool DoesAssetTypeFilterExist(const FTopLevelAssetPath& InClassPath)
 	{
 		for(const TSharedRef<FCustomClassFilterData>& CustomClassFilterData : CustomClassFilters)
 		{
-			if(CustomClassFilterData->GetFilterName() == InClassName.ToString())
+			if(CustomClassFilterData->GetClassPathName() == InClassPath)
 			{
 				return true;
 			}
@@ -308,15 +308,15 @@ public:
 	}
 
 	/** Set the check box state of the specified filter (in the filter drop down) and pin/unpin a filter widget on/from the filter bar. When a filter is pinned (was not already pinned), it is activated if requested and deactivated when unpinned.
-	 * @param	InClassName		The name of the Class the filter is associated with (must exist in the widget)
+	 * @param	InClassPath		The Class the filter is associated with (must exist in the widget)
 	 * @param	InCheckState	The CheckState to apply to the flter
 	 * @param	bEnableFilter	Whether the filter should be enabled when it is pinned
 	 */
-	void SetAssetTypeFilterCheckState(const FName& InClassName, ECheckBoxState InCheckState, bool bEnableFilter = true)
+	void SetAssetTypeFilterCheckState(const FTopLevelAssetPath& InClassPath, ECheckBoxState InCheckState, bool bEnableFilter = true)
 	{
 		for(const TSharedRef<FCustomClassFilterData>& CustomClassFilterData : CustomClassFilters)
 		{
-			if(CustomClassFilterData->GetFilterName() == InClassName.ToString())
+			if(CustomClassFilterData->GetClassPathName() == InClassPath)
 			{
 				bool FilterChecked = IsClassTypeInUse(CustomClassFilterData);
 
@@ -341,15 +341,15 @@ public:
 
 	/** Returns the check box state of the specified filter (in the filter drop down). This tells whether the filter is pinned or not on the filter bar, but not if filter is active or not.
 	 * @see		IsFilterActive().
-	 * @param	InClassName		The name of the Class the filter is associated with
+	 * @param	InClassPath		The Class the filter is associated with
 	 * @return The CheckState if a filter associated with the class name is in the filter bar, ECheckBoxState::Undetermined otherwise
 	 */
 	
-	ECheckBoxState GetAssetTypeFilterCheckState(const FName& InClassName) const
+	ECheckBoxState GetAssetTypeFilterCheckState(const FTopLevelAssetPath& InClassPath) const
 	{
 		for(const TSharedRef<FCustomClassFilterData>& CustomClassFilterData : CustomClassFilters)
 		{
-			if(CustomClassFilterData->GetFilterName() == InClassName.ToString())
+			if(CustomClassFilterData->GetClassPathName() == InClassPath)
 			{
 				return IsClassTypeInUse(CustomClassFilterData) ? ECheckBoxState::Checked : ECheckBoxState::Unchecked; 
 			}
@@ -359,13 +359,13 @@ public:
 	}
 
 	/** Returns true if the specified filter is both checked (pinned on the filter bar) and active (contributing to filter the result).
-	 *  @param	InClassName		The name of the Class the filter is associated with
+	 *  @param	InClassPath		The Class the filter is associated with
 	 */
-	bool IsAssetTypeFilterActive(const FName& InClassName) const
+	bool IsAssetTypeFilterActive(const FTopLevelAssetPath& InClassPath) const
 	{
 		for(const TSharedPtr<SAssetFilter> AssetFilter : AssetFilters)
 		{
-			if(AssetFilter->GetCustomClassFilterData()->GetFilterName() == InClassName.ToString())
+			if(AssetFilter->GetCustomClassFilterData()->GetClassPathName() == InClassPath)
 			{
 				return AssetFilter->IsEnabled();
 			}
@@ -376,14 +376,14 @@ public:
 
 	/** If a filter with the input class name is Checked (i.e visible in the bar), enable/disable it
 	 * @see SetFilterCheckState to set the check state and GetFilterCheckState to check if it is checked
-	 * @param	InClassName		The name of the Class the filter is associated with
+	 * @param	InClassPath		The Class the filter is associated with
 	 * @param	bEnable			Whether to enable or disable the filter
 	 */
-	void ToggleAssetTypeFilterEnabled(const FName& InClassName, bool bEnable)
+	void ToggleAssetTypeFilterEnabled(const FTopLevelAssetPath& InClassPath, bool bEnable)
 	{
 		for(const TSharedPtr<SAssetFilter> AssetFilter : AssetFilters)
 		{
-			if(AssetFilter->GetCustomClassFilterData()->GetFilterName() == InClassName.ToString())
+			if(AssetFilter->GetCustomClassFilterData()->GetClassPathName() == InClassPath)
 			{
 				AssetFilter->SetEnabled(bEnable);
 			}
