SFilterBar: Add the ability to save the current set of filters in a config file

diff --git a/Engine/Plugins/Editor/AssetManagerEditor/Source/AssetManagerEditor/Private/ReferenceViewer/SReferenceViewer.cpp b/Engine/Plugins/Editor/AssetManagerEditor/Source/AssetManagerEditor/Private/ReferenceViewer/SReferenceViewer.cpp
--- a/Engine/Plugins/Editor/AssetManagerEditor/Source/AssetManagerEditor/Private/ReferenceViewer/SReferenceViewer.cpp
+++ b/Engine/Plugins/Editor/AssetManagerEditor/Source/AssetManagerEditor/Private/ReferenceViewer/SReferenceViewer.cpp
@@ -130,7 +130,7 @@ void SReferenceViewer::Construct(const FArguments& InArgs)
 			{
 				GraphObj->SetCurrentFilterCollection(FilterWidget->GetAllActiveFilters());
 				GraphObj->RefilterGraph();
-				FilterWidget->SaveSettings(FString(), FString(), FString());
+				FilterWidget->SaveSettings();
 			}
 		})
 	;
@@ -782,8 +782,7 @@ void SReferenceViewer::OnUpdateFilterBar()
 
 		else
 		{
-			FString EmptyString;
-			FilterWidget->LoadSettings(EmptyString, EmptyString, EmptyString);
+			FilterWidget->LoadSettings();
 			GraphObj->SetCurrentFilterCollection(FilterWidget->GetAllActiveFilters());
 		}
 	}
diff --git a/Engine/Plugins/Editor/AssetManagerEditor/Source/AssetManagerEditor/Private/ReferenceViewer/SReferenceViewerFilterBar.cpp b/Engine/Plugins/Editor/AssetManagerEditor/Source/AssetManagerEditor/Private/ReferenceViewer/SReferenceViewerFilterBar.cpp
--- a/Engine/Plugins/Editor/AssetManagerEditor/Source/AssetManagerEditor/Private/ReferenceViewer/SReferenceViewerFilterBar.cpp
+++ b/Engine/Plugins/Editor/AssetManagerEditor/Source/AssetManagerEditor/Private/ReferenceViewer/SReferenceViewerFilterBar.cpp
@@ -4,7 +4,7 @@
 
 #include "CoreMinimal.h"
 
-void SReferenceViewerFilterBar::SaveSettings(const FString& IniFilename, const FString& IniSection, const FString& SettingsString) const
+void SReferenceViewerFilterBar::SaveSettings()
 {
 	TArray<FilterState> SaveFilters;
 	if (UReferenceViewerSettings* Settings = GetMutableDefault<UReferenceViewerSettings>())
@@ -22,7 +22,7 @@ void SReferenceViewerFilterBar::SaveSettings(const FString& IniFilename, const F
 	}
 }
 
-void SReferenceViewerFilterBar::LoadSettings(const FString& IniFilename, const FString& IniSection, const FString& SettingsString) 
+void SReferenceViewerFilterBar::LoadSettings() 
 {
 	if (UReferenceViewerSettings* Settings = GetMutableDefault<UReferenceViewerSettings>())
 	{
diff --git a/Engine/Plugins/Editor/AssetManagerEditor/Source/AssetManagerEditor/Private/ReferenceViewer/SReferenceViewerFilterBar.h b/Engine/Plugins/Editor/AssetManagerEditor/Source/AssetManagerEditor/Private/ReferenceViewer/SReferenceViewerFilterBar.h
--- a/Engine/Plugins/Editor/AssetManagerEditor/Source/AssetManagerEditor/Private/ReferenceViewer/SReferenceViewerFilterBar.h
+++ b/Engine/Plugins/Editor/AssetManagerEditor/Source/AssetManagerEditor/Private/ReferenceViewer/SReferenceViewerFilterBar.h
@@ -13,9 +13,9 @@ class SReferenceViewerFilterBar : public SFilterBar< FReferenceNodeInfo& >
 public:
 
 	/** Saves any settings to config that should be persistent between editor sessions */
-	void SaveSettings(const FString& IniFilename, const FString& IniSection, const FString& SettingsString) const override;
+	void SaveSettings() override;
 
 	/** Loads any settings to config that should be persistent between editor sessions */
-	void LoadSettings(const FString& IniFilename, const FString& IniSection, const FString& SettingsString) override;
+	void LoadSettings() override;
 
 };
\ No newline at end of file
diff --git a/Engine/Source/Developer/ToolWidgets/Public/Filters/SBasicFilterBar.h b/Engine/Source/Developer/ToolWidgets/Public/Filters/SBasicFilterBar.h
--- a/Engine/Source/Developer/ToolWidgets/Public/Filters/SBasicFilterBar.h
+++ b/Engine/Source/Developer/ToolWidgets/Public/Filters/SBasicFilterBar.h
@@ -1103,13 +1103,14 @@ protected:
 		TSharedRef<FCustomTextFilter<FilterType>> NewFilter = MakeShared<FCustomTextFilter<FilterType>>(NewTextFilter);
 
 		// Fill in the data the widget gives us
-		NewFilter->SetDisplayName(InFilterData.FilterLabel);
-		NewFilter->SetColor(InFilterData.FilterColor);
-		NewFilter->SetFilterString(InFilterData.FilterString);
+		NewFilter->SetFromCustomTextFilterData(InFilterData);
 
 		CustomTextFilters.Add(NewFilter);
+		
 		TSharedRef<SFilter> AddedFilter = AddFilterToBar(NewFilter);
+		
 		AddedFilter->SetEnabled(bApplyFilter);
+		SetFrontendFilterActive(NewFilter, bApplyFilter);
 
 		// Close the create custom filter dialog 		
 		OnCancelCustomTextFilterDialog();
@@ -1121,9 +1122,7 @@ protected:
 		// Update the filter with the data the widget has given us
 		if(InFilter)
 		{
-			InFilter->SetDisplayName(InFilterData.FilterLabel);
-			InFilter->SetColor(InFilterData.FilterColor);
-			InFilter->SetFilterString(InFilterData.FilterString);
+			InFilter->SetFromCustomTextFilterData(InFilterData);
 		}
 
 		// If the filter we modified is active in the bar, remove and re-add it to update the data
@@ -1136,7 +1135,9 @@ protected:
 				RemoveFilter(InFilter.ToSharedRef());
 
 				TSharedRef<SFilter> AddedFilter = AddFilterToBar(InFilter.ToSharedRef());
+				
 				AddedFilter->SetEnabled(bWasEnabled);
+				SetFrontendFilterActive(InFilter.ToSharedRef(), bWasEnabled);
 			}
 		}
 
@@ -1150,6 +1151,9 @@ protected:
 		{
 			CustomTextFilters.RemoveSingle(InFilter.ToSharedRef());
 			RemoveFilter(InFilter.ToSharedRef());
+
+			// We fire this delegate here, since a CustomTextFilter was removed from the list 
+			OnFilterChanged.ExecuteIfBound();
 		}
 		
 		OnCancelCustomTextFilterDialog();
@@ -1211,10 +1215,7 @@ protected:
 			return;
 		}
 
-		FCustomTextFilterData CustomTextFilterData;
-		CustomTextFilterData.FilterLabel = InFilter->GetDisplayName();
-		CustomTextFilterData.FilterColor = InFilter->GetColor();
-		CustomTextFilterData.FilterString = InFilter->GetFilterString();
+		FCustomTextFilterData CustomTextFilterData = InFilter->CreateCustomTextFilterData();
 
 		CreateCustomTextFilterWindow(CustomTextFilterData, InFilter);
 	}
@@ -1356,12 +1357,6 @@ protected:
 				);
 		}
 	}
-	
-	/** Virtual Function for subclasses to save their settings */
-	virtual void SaveSettings(const FString& IniFilename, const FString& IniSection, const FString& SettingsString) const {}
-
-	/** Virtual Function for subclasses to load their settings */
-	virtual void LoadSettings(const FString& IniFilename, const FString& IniSection, const FString& SettingsString) {}
 
 private:
 
diff --git a/Engine/Source/Developer/ToolWidgets/Public/Filters/SCustomTextFilterDialog.h b/Engine/Source/Developer/ToolWidgets/Public/Filters/SCustomTextFilterDialog.h
--- a/Engine/Source/Developer/ToolWidgets/Public/Filters/SCustomTextFilterDialog.h
+++ b/Engine/Source/Developer/ToolWidgets/Public/Filters/SCustomTextFilterDialog.h
@@ -6,9 +6,33 @@
 #include "Filters/FilterBase.h"
 #include "Core/Public/Misc/TextFilter.h"
 
+#include "SCustomTextFilterDialog.generated.h"
+
 class SColorBlock;
 class SEditableTextBox;
 
+/* Struct containing the data that SCustomTextFilterDialog is currently editing */
+USTRUCT()
+struct FCustomTextFilterData
+{
+	GENERATED_BODY()
+
+	UPROPERTY()
+	FText FilterLabel;
+	
+	UPROPERTY()
+	FText FilterString;
+
+	UPROPERTY()
+	FLinearColor FilterColor;
+
+	FCustomTextFilterData()
+		: FilterColor(FLinearColor::White)
+	{
+		
+	}
+};
+
 /* A filter that allows testing items of type FilterType against text expressions */
 template<typename FilterType>
 class FCustomTextFilter : public FFilterBase<FilterType>
@@ -25,7 +49,7 @@ public:
 	/** Returns the system name for this filter */
 	virtual FString GetName() const override
 	{
-		return FString("FCustomTextFilter");
+		return GetCommonName().ToString();
 	}
 
 	/** Returns the human readable name for this filter */
@@ -76,7 +100,7 @@ public:
 	}
 
 	/** Get the actual text this filter is using to test against */
-	FText GetFilterString()
+	FText GetFilterString() const
 	{
 		return TextFilter->GetRawFilterText();
 	}
@@ -86,6 +110,31 @@ public:
 	{
 		TextFilter->SetRawFilterText(InFilterString);
 	}
+
+	/** All FCustomTextFilters have the same internal name, this is a helper function to get that name to test against */
+	static FName GetCommonName()
+	{
+		return FName("CustomTextFilter");
+	}
+
+	/** Set the internals of this filter from an FCustomTextFilterData */
+	void SetFromCustomTextFilterData(const FCustomTextFilterData& InFilterData)
+	{
+		SetDisplayName(InFilterData.FilterLabel);
+		SetColor(InFilterData.FilterColor);
+		SetFilterString(InFilterData.FilterString);
+	}
+
+	/** Create an FCustomTextFilterData from the internals of this filter */
+	FCustomTextFilterData CreateCustomTextFilterData() const
+	{
+		FCustomTextFilterData CustomTextFilterData;
+		CustomTextFilterData.FilterLabel = GetDisplayName();
+		CustomTextFilterData.FilterColor = GetColor();
+		CustomTextFilterData.FilterString = GetFilterString();
+
+		return CustomTextFilterData;
+	}
 	
 	// Functionality that is not needed by this filter
 	
@@ -126,19 +175,7 @@ protected:
 };
 
 
-/* Struct containing the data that SCustomTextFilterDialog is currently editing */
-struct FCustomTextFilterData
-{
-	FText FilterLabel;
-	FText FilterString;
-	FLinearColor FilterColor;
 
-	FCustomTextFilterData()
-		: FilterColor(FLinearColor::White)
-	{
-		
-	}
-};
 
 class TOOLWIDGETS_API SCustomTextFilterDialog : public SCompoundWidget
 {
diff --git a/Engine/Source/Editor/ContentBrowser/Private/SFilterList.h b/Engine/Source/Editor/ContentBrowser/Private/SFilterList.h
--- a/Engine/Source/Editor/ContentBrowser/Private/SFilterList.h
+++ b/Engine/Source/Editor/ContentBrowser/Private/SFilterList.h
@@ -70,10 +70,10 @@ public:
 	void DisableFiltersThatHideItems(TArrayView<const FContentBrowserItem> ItemList);
 
 	/** Saves any settings to config that should be persistent between editor sessions */
-	void SaveSettings(const FString& IniFilename, const FString& IniSection, const FString& SettingsString) const override;
+	void SaveSettings(const FString& IniFilename, const FString& IniSection, const FString& SettingsString) const;
 
 	/** Loads any settings to config that should be persistent between editor sessions */
-	void LoadSettings(const FString& IniFilename, const FString& IniSection, const FString& SettingsString) override;
+	void LoadSettings(const FString& IniFilename, const FString& IniSection, const FString& SettingsString);
 
 	virtual FReply OnMouseButtonUp( const FGeometry& MyGeometry, const FPointerEvent& MouseEvent ) override;
 
diff --git a/Engine/Source/Editor/EditorWidgets/EditorWidgets.Build.cs b/Engine/Source/Editor/EditorWidgets/EditorWidgets.Build.cs
--- a/Engine/Source/Editor/EditorWidgets/EditorWidgets.Build.cs
+++ b/Engine/Source/Editor/EditorWidgets/EditorWidgets.Build.cs
@@ -19,6 +19,7 @@ public class EditorWidgets : ModuleRules
 				"UnrealEd",
 				"CoreUObject",
 				"ToolWidgets",
+				"EditorConfig"
 			}
 		);
 		
diff --git a/Engine/Source/Editor/EditorWidgets/Private/Filters/FilterBarConfig.cpp b/Engine/Source/Editor/EditorWidgets/Private/Filters/FilterBarConfig.cpp
new file mode 100644
--- /dev/null
+++ b/Engine/Source/Editor/EditorWidgets/Private/Filters/FilterBarConfig.cpp
@@ -0,0 +1,14 @@
+// Copyright Epic Games, Inc. All Rights Reserved.
+
+#include "Filters/FilterBarConfig.h"
+
+TObjectPtr<UFilterBarConfig> UFilterBarConfig::Instance = nullptr;
+
+void UFilterBarConfig::Initialize()
+{
+	if(!Instance)
+	{
+		Instance = NewObject<UFilterBarConfig>(); 
+		Instance->AddToRoot();
+	}
+}
\ No newline at end of file
diff --git a/Engine/Source/Editor/EditorWidgets/Private/Filters/SFilterBar.cpp b/Engine/Source/Editor/EditorWidgets/Private/Filters/SFilterBar.cpp
new file mode 100644
--- /dev/null
+++ b/Engine/Source/Editor/EditorWidgets/Private/Filters/SFilterBar.cpp
@@ -0,0 +1,40 @@
+// Copyright Epic Games, Inc. All Rights Reserved.
+
+#include "Filters/SFilterBar.h"
+
+FFilterBarSettings* FFilterBarBase::GetMutableConfig()
+{
+	if(FilterBarIdentifier.IsNone())
+	{
+		return nullptr;
+	}
+
+	UFilterBarConfig* FilterBarConfig = GetMutableDefault<UFilterBarConfig>();
+	return &UFilterBarConfig::Get()->FilterBars.FindOrAdd(FilterBarIdentifier);
+}
+
+const FFilterBarSettings* FFilterBarBase::GetConstConfig() const
+{
+	if(FilterBarIdentifier.IsNone())
+	{
+		return nullptr;
+	}
+
+	const UFilterBarConfig* FilterBarConfig = GetDefault<UFilterBarConfig>();
+	return UFilterBarConfig::Get()->FilterBars.Find(FilterBarIdentifier);
+}
+
+void FFilterBarBase::SaveConfig()
+{
+	UFilterBarConfig::Get()->SaveEditorConfig();
+}
+
+void FFilterBarBase::InitializeConfig()
+{
+	UFilterBarConfig::Initialize();
+
+	UFilterBarConfig::Get()->LoadEditorConfig();
+
+	// Call GetMutableConfig to force create a config for this filter bar if the user specified FilterBarIdentifier
+	FFilterBarSettings* FilterBarConfig = GetMutableConfig();
+}
diff --git a/Engine/Source/Editor/EditorWidgets/Public/Filters/FilterBarConfig.h b/Engine/Source/Editor/EditorWidgets/Public/Filters/FilterBarConfig.h
new file mode 100644
--- /dev/null
+++ b/Engine/Source/Editor/EditorWidgets/Public/Filters/FilterBarConfig.h
@@ -0,0 +1,75 @@
+// Copyright Epic Games, Inc. All Rights Reserved.
+
+#pragma once
+
+#include "Containers/Map.h"
+#include "EditorConfigBase.h"
+#include "Filters/SCustomTextFilterDialog.h"
+
+#include "FilterBarConfig.generated.h"
+
+USTRUCT()
+struct FCustomTextFilterState
+{
+	GENERATED_BODY()
+	
+public:
+
+	/* Whether the custom filter is checked, i.e visible in the filter bar */
+	UPROPERTY()
+	bool bIsChecked;
+
+	/* Whether the custom filter is active, i.e visible and enabled in the filter bar */
+	UPROPERTY()
+	bool bIsActive;
+
+	/* The data inside the custom text filter */
+	UPROPERTY()
+	FCustomTextFilterData FilterData;
+	
+};
+
+USTRUCT()
+struct FFilterBarSettings
+{
+	GENERATED_BODY()
+
+public:
+
+	/** Map of currently visible custom filters, along with their enabled state */
+	UPROPERTY()
+	TMap<FString, bool> CustomFilters;
+
+	/** Map of currently visible asset type filters, along with their enabled state */
+	UPROPERTY()
+	TMap<FString, bool> TypeFilters;
+
+	/** Array of custom text filters the user has created */
+	UPROPERTY()
+	TArray<FCustomTextFilterState> CustomTextFilters;
+
+	void Empty()
+	{
+		CustomFilters.Empty();
+		TypeFilters.Empty();
+		CustomTextFilters.Empty();
+	}
+};
+
+UCLASS(EditorConfig="FilterBar")
+class EDITORWIDGETS_API UFilterBarConfig : public UEditorConfigBase
+{
+	GENERATED_BODY()
+	
+public:
+
+	static void Initialize();
+	static UFilterBarConfig* Get() { return Instance; }
+	
+	UPROPERTY(meta=(EditorConfig))
+	TMap<FName, FFilterBarSettings> FilterBars;
+	
+private:
+
+	static TObjectPtr<UFilterBarConfig> Instance;
+};
diff --git a/Engine/Source/Editor/EditorWidgets/Public/Filters/SFilterBar.h b/Engine/Source/Editor/EditorWidgets/Public/Filters/SFilterBar.h
--- a/Engine/Source/Editor/EditorWidgets/Public/Filters/SFilterBar.h
+++ b/Engine/Source/Editor/EditorWidgets/Public/Filters/SFilterBar.h
@@ -8,7 +8,32 @@
 
 #include "Filters/SAssetFilterBar.h"
 #include "Filters/AssetFilter.h"
+#include "Filters/FilterBarConfig.h"
 
+/** A non-templated base class for SFilterBar, where functionality that does not depend on the type of item
+ *  being filtered lives
+ */
+class EDITORWIDGETS_API FFilterBarBase
+{
+protected:
+
+	/** Get a mutable version of this filter bar's config */
+	FFilterBarSettings* GetMutableConfig();
+
+	/** Get a const version of this filter bar's config */
+	const FFilterBarSettings* GetConstConfig() const;
+
+	/** Save this filte bar's data to the config file */
+	void SaveConfig();
+
+	/** Initialize and load this filter bar's config */
+	void InitializeConfig();
+
+protected:
+
+	/* Unique name for this filter bar */
+	FName FilterBarIdentifier;
+};
 
 /**
 * A Filter Bar widget, which can be used to filter items of type [FilterType] given a list of custom filters
@@ -35,7 +60,7 @@
 * NOTE: Use MakeAddFilterButton() to make the button that summons the dropdown showing all the filters
 */
 template<typename FilterType>
-class SFilterBar : public SAssetFilterBar<FilterType>
+class SFilterBar : public SAssetFilterBar<FilterType>, public FFilterBarBase
 {
 public:
 	using FOnFilterChanged = typename SBasicFilterBar<FilterType>::FOnFilterChanged;
@@ -43,9 +68,11 @@ public:
 	using FCompareItemWithClassNames = typename FAssetFilter<FilterType>::FCompareItemWithClassNames;
 	using FOnExtendAddFilterMenu = typename SBasicFilterBar<FilterType>::FOnExtendAddFilterMenu;
 	using FCreateTextFilter = typename SBasicFilterBar<FilterType>::FCreateTextFilter;
+	using SFilter = typename SBasicFilterBar<FilterType>::SFilter;
 	
 	SLATE_BEGIN_ARGS( SFilterBar<FilterType> )
-		: _UseDefaultAssetFilters(true)
+		: _FilterBarIdentifier(NAME_None)
+		, _UseDefaultAssetFilters(true)
 		{
 		
 		}
@@ -61,6 +88,9 @@ public:
 
 		/** Initial List of Custom Filters that will be added to the AddFilter Menu */
 		SLATE_ARGUMENT( TArray<TSharedRef<FFilterBase<FilterType>>>, CustomFilters)
+
+		/** A unique identifier for this filter bar needed to enable saving settings in a config file */
+		SLATE_ARGUMENT(FName, FilterBarIdentifier)
 	
 		/** Delegate to extend the Add Filter dropdown */
 		SLATE_EVENT( FOnExtendAddFilterMenu, OnExtendAddFilterMenu )
@@ -98,6 +128,8 @@ public:
 		
 		SAssetFilterBar<FilterType>::Construct(Args);
 
+		FilterBarIdentifier = InArgs._FilterBarIdentifier;
+		
 		// Create the dummy filter that represents all currently active Asset Type filters
 		AssetFilter = MakeShareable(new FAssetFilter<FilterType>);
 
@@ -116,6 +148,8 @@ public:
 		}
 
 		this->ActiveFilters->Add(AssetFilter);
+
+		InitializeConfig();
 	}
 	
 	/** Use this function to get all currently active filters (to filter your items)
@@ -128,6 +162,149 @@ public:
 		return this->ActiveFilters;	
 	}
 
+	virtual void SaveSettings()
+	{
+		FFilterBarSettings* FilterBarConfig = GetMutableConfig();
+		
+		if(!FilterBarConfig)
+		{
+			UE_LOG(LogSlate, Error, TEXT("SFilterBar Requires that you specify a FilterBarIdentifier to save settings"));
+			return;
+		}
+
+		// Empty the config, we are just going to re-save everything
+		FilterBarConfig->Empty();
+
+		/* Go through all the custom text filters (including unchecked ones) to save them, so that the user does not
+		 * lose the text filters they created from the menu or a saved search
+		 */
+		for (const TSharedRef<FCustomTextFilter<FilterType>>& CustomTextFilter : this->CustomTextFilters)
+		{
+			// Is the filter "checked", i.e visible in the filter bar
+			bool bIsChecked = this->IsFrontendFilterInUse(CustomTextFilter);
+
+			// Is the filter "active", i.e visible and enabled in the filter bar
+			bool bIsActive = this->IsFilterActive(CustomTextFilter);
+
+			// Get the data associated with this filter
+			FCustomTextFilterData FilterData = CustomTextFilter->CreateCustomTextFilterData();
+
+			FCustomTextFilterState FilterState;
+			FilterState.bIsChecked = bIsChecked;
+			FilterState.bIsActive = bIsActive;
+			FilterState.FilterData = FilterData;
+
+			FilterBarConfig->CustomTextFilters.Add(FilterState);
+		}
+		
+		/* For the remaining (custom and type) filters, go through the currently active (i.e visible in the filter bar)
+		 * ones to save their state, since they will be added to the filter bar programatically every time
+		 */
+		for ( const TSharedPtr<SFilter> Filter : this->Filters )
+		{
+			const FString FilterName = Filter->GetFilterName();
+
+			// Ignore custom text filters, since we saved them previously
+			if(FilterName == FCustomTextFilter<FilterType>::GetCommonName().ToString())
+			{
+				
+			}
+			// If it is a FrontendFilter
+			else if ( Filter->GetFrontendFilter().IsValid() )
+			{
+				FilterBarConfig->CustomFilters.Add(FilterName, Filter->IsEnabled());
+			}
+			// Otherwise we assume it is a type filter
+			else
+			{
+				FilterBarConfig->TypeFilters.Add(FilterName, Filter->IsEnabled());
+			}
+		}
+
+		SaveConfig();
+	}
+
+	virtual void LoadSettings()
+	{
+		const FFilterBarSettings* FilterBarConfig = GetConstConfig();
+		
+		if(!FilterBarConfig)
+		{
+			UE_LOG(LogSlate, Error, TEXT("SFilterBar Requires that you specify a FilterBarIdentifier to load settings"));
+			return;
+		}
+
+		if(!FilterBarConfig->CustomTextFilters.IsEmpty())
+		{
+			// We must have a CreateTextFilter bound if we have any custom text filters saved!!
+			check(this->CreateTextFilter.IsBound());
+		}
+
+		// Load all the custom filters (i.e FrontendFilters)
+		for ( auto FrontendFilterIt = this->AllFrontendFilters.CreateIterator(); FrontendFilterIt; ++FrontendFilterIt )
+		{
+			TSharedRef<FFilterBase<FilterType>>& FrontendFilter = *FrontendFilterIt;
+			const FString& FilterName = FrontendFilter->GetName();
+			
+			if (!this->IsFrontendFilterInUse(FrontendFilter))
+			{
+				// Try to find this type filter in our list of saved filters
+				if ( const bool* bIsActive = FilterBarConfig->CustomFilters.Find(FilterName) )
+				{
+					TSharedRef<SFilter> NewFilter = this->AddFilterToBar(FrontendFilter);
+					
+					NewFilter->SetEnabled(*bIsActive, false);
+					this->SetFrontendFilterActive(FrontendFilter, NewFilter->IsEnabled());
+				}
+			}
+		}
+		
+		// Load all the type filters
+		for(const TSharedRef<FCustomClassFilterData> &CustomClassFilter : this->CustomClassFilters)
+		{
+			if(!this->IsClassTypeInUse(CustomClassFilter))
+			{
+				const FString FilterName = CustomClassFilter->GetFilterName();
+
+				// Try to find this type filter in our list of saved filters
+				if ( const bool* bIsActive = FilterBarConfig->TypeFilters.Find(CustomClassFilter->GetFilterName()) )
+				{
+					TSharedRef<SFilter> NewFilter = this->AddAssetFilterToBar(CustomClassFilter);
+					NewFilter->SetEnabled(*bIsActive, false);
+				}
+			}
+		}
+
+		// Load all the custom text filters
+		for(const FCustomTextFilterState& FilterState : FilterBarConfig->CustomTextFilters)
+		{
+			// Create a TTextFilter for the current filter we are making using the provided delegate
+			TSharedRef<TTextFilter<FilterType>> NewTextFilter = this->CreateTextFilter.Execute().ToSharedRef();
+
+			// Create the actual text filter
+			TSharedRef<FCustomTextFilter<FilterType>> NewFilter = MakeShared<FCustomTextFilter<FilterType>>(NewTextFilter);
+
+			// Set the internals of the custom text filter from what we have saved
+			NewFilter->SetFromCustomTextFilterData(FilterState.FilterData);
+
+			// Add this to our list of custom text filters
+			this->CustomTextFilters.Add(NewFilter);
+
+			// If the filter was checked previously, add it to the filter bar
+			if(FilterState.bIsChecked)
+			{
+				TSharedRef<SFilter> AddedFilter = this->AddFilterToBar(NewFilter);
+
+				// Set the filter as active if it was previously
+				AddedFilter->SetEnabled(FilterState.bIsActive, false);
+				
+				this->SetFrontendFilterActive(NewFilter, FilterState.bIsActive);
+			}
+		}
+
+		this->OnFilterChanged.ExecuteIfBound();
+	}
+
 private:
 
 	void UpdateAssetFilter()
