Allow you to filter based on the types you can view, rather than the types you can create

#rnx

diff --git a/Engine/Source/Editor/EditorWidgets/Public/Filters/SAssetFilterBar.h b/Engine/Source/Editor/EditorWidgets/Public/Filters/SAssetFilterBar.h
--- a/Engine/Source/Editor/EditorWidgets/Public/Filters/SAssetFilterBar.h
+++ b/Engine/Source/Editor/EditorWidgets/Public/Filters/SAssetFilterBar.h
@@ -536,7 +536,7 @@ protected:
 		};
 		AssetTypeActionsList.Sort( FCompareIAssetTypeActions() );
 
-		const TSharedRef<FPathPermissionList>& AssetClassPermissionList = AssetToolsModule.Get().GetAssetClassPathPermissionList(EAssetClassAction::CreateAsset);
+		const TSharedRef<FPathPermissionList>& AssetClassPermissionList = AssetToolsModule.Get().GetAssetClassPathPermissionList(EAssetClassAction::ViewAsset);
 
 		// For every asset type, convert it to an FCustomClassFilterData and add it to the list
 		for (int32 ClassIdx = 0; ClassIdx < AssetTypeActionsList.Num(); ++ClassIdx)
