SFilterBar: Fix issue with uninitialized bools in FCustomTextFilterState

diff --git a/Engine/Source/Editor/EditorWidgets/Public/Filters/FilterBarConfig.h b/Engine/Source/Editor/EditorWidgets/Public/Filters/FilterBarConfig.h
--- a/Engine/Source/Editor/EditorWidgets/Public/Filters/FilterBarConfig.h
+++ b/Engine/Source/Editor/EditorWidgets/Public/Filters/FilterBarConfig.h
@@ -17,11 +17,11 @@ public:
 
 	/* Whether the custom filter is checked, i.e visible in the filter bar */
 	UPROPERTY()
-	bool bIsChecked;
+	bool bIsChecked = false;
 
 	/* Whether the custom filter is active, i.e visible and enabled in the filter bar */
 	UPROPERTY()
-	bool bIsActive;
+	bool bIsActive = false;
 
 	/* The data inside the custom text filter */
 	UPROPERTY()
