SAssetFilterBar: Update the default Asset Type Filters when the AssetClassPermissionList changes

diff --git a/Engine/Source/Editor/EditorWidgets/Public/Filters/SAssetFilterBar.h b/Engine/Source/Editor/EditorWidgets/Public/Filters/SAssetFilterBar.h
--- a/Engine/Source/Editor/EditorWidgets/Public/Filters/SAssetFilterBar.h
+++ b/Engine/Source/Editor/EditorWidgets/Public/Filters/SAssetFilterBar.h
@@ -119,6 +119,7 @@ public:
 	{
 		bUseDefaultAssetFilters = InArgs._UseDefaultAssetFilters;
 		CustomClassFilters = InArgs._CustomClassFilters;
+		UserCustomClassFilters = InArgs._CustomClassFilters;
 		
 		typename SBasicFilterBar<FilterType>::FArguments Args;
 		Args._OnFilterChanged = InArgs._OnFilterChanged;
@@ -128,8 +129,16 @@ public:
 		Args._FilterSearchBox = InArgs._FilterSearchBox;
 		
 		SBasicFilterBar<FilterType>::Construct(Args);
-		
+
+		// Create the Asset Type Action filters if requested
 		CreateAssetTypeActionFilters();
+
+		FAssetToolsModule& AssetToolsModule = FModuleManager::LoadModuleChecked<FAssetToolsModule>(TEXT("AssetTools"));
+		const TSharedRef<FPathPermissionList>& AssetClassPermissionList = AssetToolsModule.Get().GetAssetClassPathPermissionList(AssetClassAction);
+
+		// Re-create the Asset Type Action filters whenever the permission list changes
+		AssetClassPermissionList->OnFilterChanged().AddSP(this, &SAssetFilterBar<FilterType>::CreateAssetTypeActionFilters);
+
 	}
 
 protected:
@@ -505,9 +514,14 @@ protected:
 			return;
 		}
 
-		FAssetToolsModule& AssetToolsModule = FModuleManager::LoadModuleChecked<FAssetToolsModule>(TEXT("AssetTools"));
-
+		// Empty the existing categories and filters, we will re-add them based on the permission list
 		AssetFilterCategories.Empty();
+		CustomClassFilters.Empty();
+
+		// Re-add the CustomClassFilters added by the user manually, these don't get tested against the permission list
+		CustomClassFilters.Append(UserCustomClassFilters);
+
+		FAssetToolsModule& AssetToolsModule = FModuleManager::LoadModuleChecked<FAssetToolsModule>(TEXT("AssetTools"));
 
 		// Add the Basic category
 		AssetFilterCategories.Add(EAssetTypeCategories::Basic, MakeShareable(new FFilterCategory(LOCTEXT("BasicFilter", "Basic"), LOCTEXT("BasicFilterTooltip", "Filter by basic assets."))));
@@ -522,8 +536,9 @@ protected:
 			AssetFilterCategories.Add(AdvancedAssetCategory.CategoryType, MakeShareable(new FFilterCategory(AdvancedAssetCategory.CategoryName, Tooltip)));
 		}
 
-		// Get the browser type maps
 		TArray<TWeakPtr<IAssetTypeActions>> AssetTypeActionsList;
+		
+		// Get the browser type maps
 		AssetToolsModule.Get().GetAssetTypeActionsList(AssetTypeActionsList);
 		
 		// Sort the list
@@ -536,7 +551,7 @@ protected:
 		};
 		AssetTypeActionsList.Sort( FCompareIAssetTypeActions() );
 
-		const TSharedRef<FPathPermissionList>& AssetClassPermissionList = AssetToolsModule.Get().GetAssetClassPathPermissionList(EAssetClassAction::ViewAsset);
+		const TSharedRef<FPathPermissionList>& AssetClassPermissionList = AssetToolsModule.Get().GetAssetClassPathPermissionList(AssetClassAction);
 
 		// For every asset type, convert it to an FCustomClassFilterData and add it to the list
 		for (int32 ClassIdx = 0; ClassIdx < AssetTypeActionsList.Num(); ++ClassIdx)
@@ -912,11 +927,17 @@ protected:
 	/** List of custom Class Filters that will be shown in the filter bar */
 	TArray<TSharedRef<FCustomClassFilterData>> CustomClassFilters;
 
+	/** Copy of Custom Class Filters that were provided by the user, and not autopopulated from AssetTypeActions */
+	TArray<TSharedRef<FCustomClassFilterData>> UserCustomClassFilters;
+
 	TMap<EAssetTypeCategories::Type, TSharedPtr<FFilterCategory>> AssetFilterCategories;
 	
 	/** Whether the filter bar provides the default Asset Filters */
 	bool bUseDefaultAssetFilters;
 
+	/** The AssetClassAction used to get the permission list from the Asset Tools Module */
+	EAssetClassAction AssetClassAction = EAssetClassAction::ViewAsset;
+
 };
 
 #undef LOCTEXT_NAMESPACE
\ No newline at end of file
