Outliner: Fix an issue causing unloaded actors to not show up

diff --git a/Engine/Source/Editor/EditorWidgets/Public/Filters/AssetFilter.h b/Engine/Source/Editor/EditorWidgets/Public/Filters/AssetFilter.h
--- a/Engine/Source/Editor/EditorWidgets/Public/Filters/AssetFilter.h
+++ b/Engine/Source/Editor/EditorWidgets/Public/Filters/AssetFilter.h
@@ -101,7 +101,10 @@ public:
 			// Skip the current item if it wasn't converted
 			if(!bWasConverted)
 			{
-				return false;
+				/* If the current item does not want to be compared with Asset Type Filters, return true if there are no
+				 * asset type filters active, since we want to show it then. Return false if there are any filters active
+				 * since we only want to show items that pass the filters if so */
+				return CompiledBackendFilter.IsEmpty();
 			}
 			
 			// Use the AssetRegistry to check if the Item passes the compiled filter
