Multi-line editable text boxes no longer call IPropertyUtilities::RequestRefresh() when adding or removing lines, which caused textboxes created in property row generators to fail to accept input past a single line.

This was a bug exposed because FPropertyRowGenerator previously ignored RequestRefresh() calls entirely, but now does a deferred refresh.

Since this was the only use of it, IEditableTextProperty and all its implementations have also had their RequestRefresh() methods removed.

diff --git a/Engine/Plugins/Runtime/Metasound/Source/MetasoundEditor/Private/MetasoundNodeDetailCustomization.h b/Engine/Plugins/Runtime/Metasound/Source/MetasoundEditor/Private/MetasoundNodeDetailCustomization.h
--- a/Engine/Plugins/Runtime/Metasound/Source/MetasoundEditor/Private/MetasoundNodeDetailCustomization.h
+++ b/Engine/Plugins/Runtime/Metasound/Source/MetasoundEditor/Private/MetasoundNodeDetailCustomization.h
@@ -61,7 +61,6 @@ namespace Metasound
 			virtual bool IsReadOnly() const override { return false; }
 			virtual int32 GetNumTexts() const override { return 1; }
 			virtual bool IsValidText(const FText& InText, FText& OutErrorMsg) const override { return true; }
-			virtual void RequestRefresh() override { }
 
 			virtual FText GetToolTipText() const override
 			{
diff --git a/Engine/Source/Editor/DetailCustomizations/Private/TextCustomization.cpp b/Engine/Source/Editor/DetailCustomizations/Private/TextCustomization.cpp
--- a/Engine/Source/Editor/DetailCustomizations/Private/TextCustomization.cpp
+++ b/Engine/Source/Editor/DetailCustomizations/Private/TextCustomization.cpp
@@ -102,14 +102,6 @@ namespace
 		}
 #endif // USE_STABLE_LOCALIZATION_KEYS
 
-		virtual void RequestRefresh() override
-		{
-			if (PropertyUtilities.IsValid())
-			{
-				PropertyUtilities->RequestRefresh();
-			}
-		}
-
 	private:
 		TSharedRef<IPropertyHandle> PropertyHandle;
 		TSharedPtr<IPropertyUtilities> PropertyUtilities;
diff --git a/Engine/Source/Editor/EditorWidgets/Private/STextPropertyEditableTextBox.cpp b/Engine/Source/Editor/EditorWidgets/Private/STextPropertyEditableTextBox.cpp
--- a/Engine/Source/Editor/EditorWidgets/Private/STextPropertyEditableTextBox.cpp
+++ b/Engine/Source/Editor/EditorWidgets/Private/STextPropertyEditableTextBox.cpp
@@ -953,16 +953,6 @@ FReply STextPropertyEditableTextBox::OnFocusReceived(const FGeometry& MyGeometry
 	return FReply::Handled().SetUserFocus(PrimaryWidget.ToSharedRef(), InFocusEvent.GetCause());
 }
 
-void STextPropertyEditableTextBox::Tick(const FGeometry& AllottedGeometry, const double InCurrentTime, const float InDeltaTime)
-{
-	const float CurrentHeight = AllottedGeometry.GetLocalSize().Y;
-	if (bIsMultiLine && PreviousHeight.IsSet() && PreviousHeight.GetValue() != CurrentHeight)
-	{
-		EditableTextProperty->RequestRefresh();
-	}
-	PreviousHeight = CurrentHeight;
-}
-
 bool STextPropertyEditableTextBox::CanEdit() const
 {
 	const bool bIsReadOnly = FTextLocalizationManager::Get().IsLocalizationLocked() || EditableTextProperty->IsReadOnly();
diff --git a/Engine/Source/Editor/EditorWidgets/Public/STextPropertyEditableTextBox.h b/Engine/Source/Editor/EditorWidgets/Public/STextPropertyEditableTextBox.h
--- a/Engine/Source/Editor/EditorWidgets/Public/STextPropertyEditableTextBox.h
+++ b/Engine/Source/Editor/EditorWidgets/Public/STextPropertyEditableTextBox.h
@@ -92,9 +92,6 @@ public:
 	virtual void GetStableTextId(const int32 InIndex, const ETextPropertyEditAction InEditAction, const FString& InTextSource, const FString& InProposedNamespace, const FString& InProposedKey, FString& OutStableNamespace, FString& OutStableKey) const = 0;
 #endif // USE_STABLE_LOCALIZATION_KEYS
 
-	/** Request a refresh of the property UI (eg, due to a size change) */
-	virtual void RequestRefresh() = 0;
-
 protected:
 #if USE_STABLE_LOCALIZATION_KEYS
 	/** Get the localization ID we should use for the given object, and the given text instance */
@@ -208,7 +205,6 @@ public:
 	void Construct(const FArguments& Arguments, const TSharedRef<IEditableTextProperty>& InEditableTextProperty);
 	virtual bool SupportsKeyboardFocus() const override;
 	virtual FReply OnFocusReceived(const FGeometry& MyGeometry, const FFocusEvent& InFocusEvent) override;
-	virtual void Tick(const FGeometry& AllottedGeometry, const double InCurrentTime, const float InDeltaTime) override;
 
 private:
 	void GetDesiredWidth(float& OutMinDesiredWidth, float& OutMaxDesiredWidth);
@@ -257,9 +253,7 @@ private:
 
 	TSharedPtr<SEditableTextBox> KeyEditableTextBox;
 
-	TOptional<float> PreviousHeight;
-
-	bool bIsMultiLine;
+	bool bIsMultiLine = false;
 
 	static FText MultipleValuesText;
 };
diff --git a/Engine/Source/Editor/GraphEditor/Private/KismetPins/SGraphPinText.cpp b/Engine/Source/Editor/GraphEditor/Private/KismetPins/SGraphPinText.cpp
--- a/Engine/Source/Editor/GraphEditor/Private/KismetPins/SGraphPinText.cpp
+++ b/Engine/Source/Editor/GraphEditor/Private/KismetPins/SGraphPinText.cpp
@@ -96,10 +96,6 @@ namespace
 		}
 #endif // USE_STABLE_LOCALIZATION_KEYS
 
-		virtual void RequestRefresh() override
-		{
-		}
-
 	private:
 		FEdGraphPinReference GraphPinObjRef;
 	};
diff --git a/Engine/Source/Editor/Kismet/Private/UserDefinedEnumEditor.cpp b/Engine/Source/Editor/Kismet/Private/UserDefinedEnumEditor.cpp
--- a/Engine/Source/Editor/Kismet/Private/UserDefinedEnumEditor.cpp
+++ b/Engine/Source/Editor/Kismet/Private/UserDefinedEnumEditor.cpp
@@ -114,10 +114,6 @@ public:
 	}
 #endif // USE_STABLE_LOCALIZATION_KEYS
 
-	virtual void RequestRefresh() override
-	{
-	}
-
 	bool CausedChange() const
 	{
 		return bCausedChange;
@@ -206,10 +202,6 @@ public:
 	}
 #endif // USE_STABLE_LOCALIZATION_KEYS
 
-	virtual void RequestRefresh() override
-	{
-	}
-
 	bool CausedChange() const
 	{
 		return bCausedChange;
