SEnumComboBox: VisibleEnums argument

diff --git a/Engine/Source/Editor/EditorWidgets/Private/SEnumCombo.cpp b/Engine/Source/Editor/EditorWidgets/Private/SEnumCombo.cpp
--- a/Engine/Source/Editor/EditorWidgets/Private/SEnumCombo.cpp
+++ b/Engine/Source/Editor/EditorWidgets/Private/SEnumCombo.cpp
@@ -34,16 +34,28 @@ void SEnumComboBox::Construct(const FArguments& InArgs, const UEnum* InEnum)
 	bUpdatingSelectionInternally = false;
 	bIsBitflagsEnum = Enum->HasMetaData(TEXT("Bitflags"));
 
+	TArray<int32> EnumValueSubset = InArgs._EnumValueSubset;
+	if(EnumValueSubset.IsEmpty())
+	{
+		EnumValueSubset.Reserve(Enum->NumEnums() - 1);
+		for (int32 i = 0; i < Enum->NumEnums() - 1; i++)
+		{
+			// Note: SEnumComboBox API prior to bitflags only supports 32 bit values, truncating the value here to keep the old API.
+			const int32 Value = Enum->GetValueByIndex(i);
+			EnumValueSubset.Add(Value);
+		}
+	}
+
 	if (bIsBitflagsEnum)
 	{
 		const bool bUseEnumValuesAsMaskValues = Enum->GetBoolMetaData(UseEnumValuesAsMaskValuesInEditorName);
 		const int32 BitmaskBitCount = sizeof(int32) << 3;
 
-		for (int32 i = 0; i < Enum->NumEnums() - 1; i++)
+		for (int32 i = 0; i < EnumValueSubset.Num(); i++)
 		{
-			// Note: SEnumComboBox API prior to bitflags only supports 32 bit values, truncating the value here to keep the old API.
-			int32 Value = Enum->GetValueByIndex(i);
-			const bool bIsHidden = Enum->HasMetaData(TEXT("Hidden"), i);
+			int32 Value = EnumValueSubset[i];
+			const int32 Index = Enum->GetIndexByValue(Value);
+			const bool bIsHidden = Enum->HasMetaData(TEXT("Hidden"), Index);
 			if (Value >= 0 && !bIsHidden)
 			{
 				if (bUseEnumValuesAsMaskValues)
@@ -62,24 +74,26 @@ void SEnumComboBox::Construct(const FArguments& InArgs, const UEnum* InEnum)
 					Value = 1 << Value;
 				}
 
-				FText DisplayName = Enum->GetDisplayNameTextByIndex(i);
-				FText TooltipText = Enum->GetToolTipTextByIndex(i);
+				FText DisplayName = Enum->GetDisplayNameTextByIndex(Index);
+				FText TooltipText = Enum->GetToolTipTextByIndex(Index);
 				if (TooltipText.IsEmpty())
 				{
 					TooltipText = FText::Format(LOCTEXT("BitmaskDefaultFlagToolTipText", "Toggle {0} on/off"), DisplayName);
 				}
 
-				VisibleEnums.Emplace(i, Value, DisplayName, TooltipText);
+				VisibleEnums.Emplace(Index, Value, DisplayName, TooltipText);
 			}
 		}
 	}
 	else
 	{
-		for (int32 i = 0; i < Enum->NumEnums() - 1; i++)
+		for (int32 i = 0; i < EnumValueSubset.Num(); i++)
 		{
-			if (Enum->HasMetaData(TEXT("Hidden"), i) == false)
+			const int32 Value = EnumValueSubset[i];
+			const int32 Index = Enum->GetIndexByValue(Value);
+			if (Enum->HasMetaData(TEXT("Hidden"), Index) == false)
 			{
-				VisibleEnums.Emplace(i, Enum->GetValueByIndex(i), Enum->GetDisplayNameTextByIndex(i), Enum->GetToolTipTextByIndex(i));
+				VisibleEnums.Emplace(Index, Enum->GetValueByIndex(Index), Enum->GetDisplayNameTextByIndex(Index), Enum->GetToolTipTextByIndex(Index));
 			}
 		}
 	}
diff --git a/Engine/Source/Editor/EditorWidgets/Public/SEnumCombo.h b/Engine/Source/Editor/EditorWidgets/Public/SEnumCombo.h
--- a/Engine/Source/Editor/EditorWidgets/Public/SEnumCombo.h
+++ b/Engine/Source/Editor/EditorWidgets/Public/SEnumCombo.h
@@ -38,6 +38,7 @@ public:
 		SLATE_EVENT(FOnEnumSelectionChanged, OnEnumSelectionChanged)
 		SLATE_EVENT(FOnGetToolTipForValue, OnGetToolTipForValue)
 		SLATE_STYLE_ARGUMENT(FButtonStyle, ButtonStyle)
+		SLATE_ARGUMENT(TArray<int32>, EnumValueSubset)
 	SLATE_END_ARGS()
 
 	void Construct(const FArguments& InArgs, const UEnum* InEnum);
