SAssetFilterBar: Fix an issue that caused you to have duplicate filters for the same types in some situations

diff --git a/Engine/Source/Editor/EditorWidgets/Public/Filters/SAssetFilterBar.h b/Engine/Source/Editor/EditorWidgets/Public/Filters/SAssetFilterBar.h
--- a/Engine/Source/Editor/EditorWidgets/Public/Filters/SAssetFilterBar.h
+++ b/Engine/Source/Editor/EditorWidgets/Public/Filters/SAssetFilterBar.h
@@ -519,6 +519,11 @@ protected:
 			return CustomClassFilter;
 		}
 
+		void SetCustomClassFilterData(const TSharedRef<FCustomClassFilterData>& InCustomClassFilterData)
+		{
+			CustomClassFilter = InCustomClassFilterData;
+		}
+
 		/** Returns the display name for this filter */
 		virtual FText GetFilterDisplayName() const override
 		{
@@ -860,6 +865,39 @@ protected:
 		// Do a second pass through all the CustomClassFilters with AssetTypeActions to update their categories
 		UpdateAssetTypeActionCategories();
 
+		// Update/remove any asset filters that already exist and use AssetTypeActions
+		TArray<TSharedRef<SAssetFilter>> AssetFiltersToRemove;
+		
+		for(TSharedRef<SAssetFilter>& AssetFilter : AssetFilters)
+		{
+			// We only care about filters that are based on asset type actions
+			if(!AssetFilter->GetCustomClassFilterData()->GetAssetTypeActions())
+			{
+				continue;
+			}
+
+			// Try and find the FCustomClassFilterData for the current filter using the new list we just created in CustomClassFilters
+			TSharedRef<FCustomClassFilterData>* FoundCustomClassFilterData = CustomClassFilters.FindByPredicate([AssetFilter](const TSharedRef<FCustomClassFilterData>& CustomClassFilterData)
+			{
+				return CustomClassFilterData->GetAssetTypeActions() == AssetFilter->GetCustomClassFilterData()->GetAssetTypeActions();
+			});
+
+			// If it exists, update the filter to let it know
+			if(FoundCustomClassFilterData)
+			{
+				AssetFilter->SetCustomClassFilterData(*FoundCustomClassFilterData);
+			}
+			// Otherwise this filter had an AssetTypeAction that doesn't pass the new permission list, remove it
+			else
+			{
+				AssetFiltersToRemove.Add(AssetFilter);
+			}
+		}
+
+		for(TSharedRef<SAssetFilter>& AssetFilter : AssetFiltersToRemove)
+		{
+			RemoveFilter(AssetFilter);
+		}
 	}
 
 	void UpdateAssetTypeActionCategories()
@@ -1111,7 +1149,7 @@ protected:
 	/** Handler to determine the "checked" state of class filter in the filter dropdown */
 	bool IsClassTypeInUse(TSharedPtr<FCustomClassFilterData> Class) const
 	{
-		for (const TSharedPtr<SAssetFilter> AssetFilter : this->AssetFilters)
+		for (const TSharedPtr<SAssetFilter>& AssetFilter : this->AssetFilters)
 		{
 			if (AssetFilter.IsValid() && AssetFilter->GetCustomClassFilterData() == Class)
 			{
@@ -1122,7 +1160,7 @@ protected:
 		return false;
 	}
 
- 	/** Handler for when filter by type category is selected */
+	/** Handler for when filter by type category is selected */
  	void FilterByTypeCategoryClicked(TSharedPtr<FFilterCategory> TypeCategory, TArray<TSharedPtr<FCustomClassFilterData>> Classes)
 	{
 		bool bFullCategoryInUse = IsTypeCategoryInUse(TypeCategory, Classes);
