- Added bool attribute to SDropTarget to stop AllowDrop function call from executing every frame for overlay visibility.
This results in the drop target overlay visibility only appearing when actually hovering the drop target when set to true.
This is useful for drop targets that are used in conjunction with other widgets with heavy drag & drop usage such as tree views.
Defaults to false to maintain previous behavior.

- Added OnDragEnter & OnDragLeave events to allow proper handling of tooltips (i.e. setting a tooltip description when hovering a drop target, and resetting it when leaving again)
AllowDrop could not be used for this since it was executed every frame for overlay visibility checks.

diff --git a/Engine/Source/Editor/EditorWidgets/Private/SDropTarget.cpp b/Engine/Source/Editor/EditorWidgets/Private/SDropTarget.cpp
--- a/Engine/Source/Editor/EditorWidgets/Private/SDropTarget.cpp
+++ b/Engine/Source/Editor/EditorWidgets/Private/SDropTarget.cpp
@@ -25,7 +25,10 @@ void SDropTarget::Construct(const FArguments& InArgs)
 	DroppedEvent = InArgs._OnDropped;
 	AllowDropEvent = InArgs._OnAllowDrop;
 	IsRecognizedEvent = InArgs._OnIsRecognized;
-
+	OnDragEnterEvent = InArgs._OnDragEnter;
+	OnDragLeaveEvent = InArgs._OnDragLeave;
+	bOnlyRecognizeOnDragEnter = InArgs._bOnlyRecognizeOnDragEnter;
+	
 	bIsDragEventRecognized = false;
 	bAllowDrop = false;
 	bIsDragOver = false;
@@ -64,9 +67,19 @@ EVisibility SDropTarget::GetDragOverlayVisibility() const
 {
 	if ( FSlateApplication::Get().IsDragDropping() )
 	{
-		if ( AllowDrop(FSlateApplication::Get().GetDragDroppingContent()) || (bIsDragOver && bIsDragEventRecognized) )
+		bool bCheckForAllowDrop = true;
+
+		if(bOnlyRecognizeOnDragEnter.Get() && !bIsDragOver)
+		{
+			bCheckForAllowDrop = false;
+		}
+		
+		if(bCheckForAllowDrop)
 		{
-			return EVisibility::HitTestInvisible;
+			if(AllowDrop(FSlateApplication::Get().GetDragDroppingContent()) || (bIsDragOver && bIsDragEventRecognized))
+			{
+				return EVisibility::HitTestInvisible;
+			}
 		}
 	}
 
@@ -83,7 +96,6 @@ bool SDropTarget::AllowDrop(TSharedPtr<FDragDropOperation> DragDropOperation) co
 {
 	bAllowDrop = OnAllowDrop(DragDropOperation);
 	bIsDragEventRecognized = OnIsRecognized(DragDropOperation) || bAllowDrop;
-
 	return bAllowDrop;
 }
 
@@ -135,6 +147,8 @@ void SDropTarget::OnDragEnter(const FGeometry& MyGeometry, const FDragDropEvent&
 	// initially we dont recognize this event
 	bIsDragEventRecognized = false;
 	bIsDragOver = true;
+
+	OnDragEnterEvent.ExecuteIfBound(DragDropEvent);
 }
 
 void SDropTarget::OnDragLeave(const FDragDropEvent& DragDropEvent)
@@ -145,6 +159,7 @@ void SDropTarget::OnDragLeave(const FDragDropEvent& DragDropEvent)
 	bAllowDrop = false;
 
 	bIsDragOver = false;
+	OnDragLeaveEvent.ExecuteIfBound(DragDropEvent);
 }
 
 int32 SDropTarget::OnPaint(const FPaintArgs& Args, const FGeometry& AllottedGeometry, const FSlateRect& MyCullingRect, FSlateWindowElementList& OutDrawElements, int32 LayerId, const FWidgetStyle& InWidgetStyle, bool bParentEnabled) const
diff --git a/Engine/Source/Editor/EditorWidgets/Public/SDropTarget.h b/Engine/Source/Editor/EditorWidgets/Public/SDropTarget.h
--- a/Engine/Source/Editor/EditorWidgets/Public/SDropTarget.h
+++ b/Engine/Source/Editor/EditorWidgets/Public/SDropTarget.h
@@ -33,6 +33,7 @@ class EDITORWIDGETS_API SDropTarget : public SCompoundWidget
 {
 public:
 	DECLARE_DELEGATE_RetVal_OneParam(bool, FVerifyDrag, TSharedPtr<FDragDropOperation>);
+	DECLARE_DELEGATE_OneParam(FOnDragAction, const FDragDropEvent&);
 	DECLARE_DELEGATE_RetVal_OneParam(FReply, FOnDropDeprecated, TSharedPtr<FDragDropOperation>);
 
 	SLATE_BEGIN_ARGS(SDropTarget)
@@ -41,6 +42,7 @@ public:
 		, _VerticalImage(FAppStyle::GetBrush("WideDash.Vertical"))
 		, _HorizontalImage(FAppStyle::GetBrush("WideDash.Horizontal"))
 		, _BackgroundImage(FAppStyle::GetBrush("DropTarget.Background"))
+		, _bOnlyRecognizeOnDragEnter(false)
 	{ }
 	
 		UE_DEPRECATED(5.0, "BackgroundColor has been removed. You may alter the background brush to get the same effect.")
@@ -72,6 +74,10 @@ public:
 		SLATE_EVENT(FVerifyDrag, OnAllowDrop)
 		/** Called to check if an asset is acceptable for dropping */
 		SLATE_EVENT(FVerifyDrag, OnIsRecognized)
+		SLATE_EVENT(FOnDragAction, OnDragEnter)
+		SLATE_EVENT(FOnDragAction, OnDragLeave)
+		/** When this is true, the drop target will only get recognized when entering while drag & dropping. */
+		SLATE_ATTRIBUTE(bool, bOnlyRecognizeOnDragEnter)
 
 		FOnDrop ConvertOnDropFn(const FOnDropDeprecated& LegacyDelegate)
 		{
@@ -122,6 +128,10 @@ private:
 	FVerifyDrag AllowDropEvent;
 	/** Delegate to call to check validity of the asset */
 	FVerifyDrag IsRecognizedEvent;
+	FOnDragAction OnDragEnterEvent;
+	FOnDragAction OnDragLeaveEvent;
+	/** Attribute to check if the drop target should only be useable when actually dragging over it. */
+	TAttribute<bool> bOnlyRecognizeOnDragEnter;
 
 	/** The color of the vertical/horizontal images when the drop data is valid */
 	FSlateColor ValidColor;
