Ensure we always edit the source string of an FText

This ensures that you can't stomp the text on your assets even if game localization data gets loaded without locking down the editable text fields

#rnx

diff --git a/Engine/Source/Editor/EditorWidgets/Private/STextPropertyEditableTextBox.cpp b/Engine/Source/Editor/EditorWidgets/Private/STextPropertyEditableTextBox.cpp
--- a/Engine/Source/Editor/EditorWidgets/Private/STextPropertyEditableTextBox.cpp
+++ b/Engine/Source/Editor/EditorWidgets/Private/STextPropertyEditableTextBox.cpp
@@ -1046,10 +1046,27 @@ FText STextPropertyEditableTextBox::GetToolTipText() const
 				const FString PackageNamespace = TextNamespaceUtil::ExtractPackageNamespace(Namespace);
 				const FString TextNamespace = TextNamespaceUtil::StripPackageNamespace(Namespace);
 
-				LocalizedTextToolTip = FText::Format(
-					LOCTEXT("LocalizedTextToolTipFmt", "--- Localized Text ---\nPackage: {0}\nNamespace: {1}\nKey: {2}\nSource: {3}"), 
-					FText::FromString(PackageNamespace), FText::FromString(TextNamespace), FText::FromString(Key), FText::FromString(*SourceString)
-					);
+				FFormatNamedArguments LocalizedTextToolTipArgs;
+				LocalizedTextToolTipArgs.Add(TEXT("Package"), FText::FromString(PackageNamespace));
+				LocalizedTextToolTipArgs.Add(TEXT("Namespace"), FText::FromString(TextNamespace));
+				LocalizedTextToolTipArgs.Add(TEXT("Key"), FText::FromString(Key));
+				LocalizedTextToolTipArgs.Add(TEXT("Source"), FText::FromString(*SourceString));
+				LocalizedTextToolTipArgs.Add(TEXT("Display"), TextValue);
+
+				if (SourceString->Equals(TextValue.ToString(), ESearchCase::CaseSensitive))
+				{
+					LocalizedTextToolTip = FText::Format(
+						LOCTEXT("LocalizedTextNoDisplayToolTipFmt", "--- Localized Text ---\nPackage: {Package}\nNamespace: {Namespace}\nKey: {Key}\nSource: {Source}"),
+						LocalizedTextToolTipArgs
+						);
+				}
+				else
+				{
+					LocalizedTextToolTip = FText::Format(
+						LOCTEXT("LocalizedTextWithDisplayToolTipFmt", "--- Localized Text ---\nPackage: {Package}\nNamespace: {Namespace}\nKey: {Key}\nSource: {Source}\nDisplay: {Display}"),
+						LocalizedTextToolTipArgs
+						);
+				}
 			}
 		}
 	}
@@ -1072,7 +1089,7 @@ FText STextPropertyEditableTextBox::GetToolTipText() const
 		return LocalizedTextToolTip;
 	}
 
-	return FText::Format(LOCTEXT("ToolTipCompleteFmt", "{0}\n\n{1}"), BaseToolTipText, LocalizedTextToolTip);
+	return FText::Format(FText::AsCultureInvariant(TEXT("{0}\n\n{1}")), BaseToolTipText, LocalizedTextToolTip);
 }
 
 FText STextPropertyEditableTextBox::GetTextValue() const
@@ -1083,6 +1100,17 @@ FText STextPropertyEditableTextBox::GetTextValue() const
 	if (NumTexts == 1)
 	{
 		TextValue = EditableTextProperty->GetText(0);
+
+		if (const FString* SourceString = FTextInspector::GetSourceString(TextValue);
+			SourceString && !FTextLocalizationManager::Get().IsLocalizationLocked())
+		{
+			// We should always edit the source string, but if the source string matches the current 
+			// display string then we can avoid making a temporary text from the source string
+			if (!SourceString->Equals(TextValue.ToString(), ESearchCase::CaseSensitive))
+			{
+				TextValue = FText::AsCultureInvariant(*SourceString);
+			}
+		}
 	}
 	else if (NumTexts > 1)
 	{
