Avoid editing the source string if it's empty

That could happen if using FText::FindText with a null source string

#rnx

diff --git a/Engine/Source/Editor/EditorWidgets/Private/STextPropertyEditableTextBox.cpp b/Engine/Source/Editor/EditorWidgets/Private/STextPropertyEditableTextBox.cpp
--- a/Engine/Source/Editor/EditorWidgets/Private/STextPropertyEditableTextBox.cpp
+++ b/Engine/Source/Editor/EditorWidgets/Private/STextPropertyEditableTextBox.cpp
@@ -1106,7 +1106,7 @@ FText STextPropertyEditableTextBox::GetTextValue() const
 		{
 			// We should always edit the source string, but if the source string matches the current 
 			// display string then we can avoid making a temporary text from the source string
-			if (!SourceString->Equals(TextValue.ToString(), ESearchCase::CaseSensitive))
+			if (!SourceString->IsEmpty() && !SourceString->Equals(TextValue.ToString(), ESearchCase::CaseSensitive))
 			{
 				TextValue = FText::AsCultureInvariant(*SourceString);
 			}
