Editor - Fixing the asset definition filtering so that we correctly merge and then ignore unmerged filters so that Animation Blueprints when you filter for Blueprints.

diff --git a/Engine/Source/Editor/EditorWidgets/Private/Filters/CustomClassFilterData.cpp b/Engine/Source/Editor/EditorWidgets/Private/Filters/CustomClassFilterData.cpp
--- a/Engine/Source/Editor/EditorWidgets/Private/Filters/CustomClassFilterData.cpp
+++ b/Engine/Source/Editor/EditorWidgets/Private/Filters/CustomClassFilterData.cpp
@@ -11,7 +11,7 @@
 FCustomClassFilterData::FCustomClassFilterData(UAssetDefinition* InAssetDefinition, const FAssetFilterData& InAssetTypeAction)
 	: AssetDefinitionPtr(InAssetDefinition)
 {
-	Class = InAssetDefinition->GetAssetClass()->GetClass();
+	Class = InAssetDefinition->GetAssetClass().Get();
 	FilterName = InAssetTypeAction.Name;
 	FilterDisplayName = InAssetTypeAction.DisplayText;
 	ClassPathName = InAssetDefinition->GetAssetClass().ToSoftObjectPath().GetAssetPath();
@@ -19,7 +19,7 @@ FCustomClassFilterData::FCustomClassFilterData(UAssetDefinition* InAssetDefiniti
 	Filter = InAssetTypeAction.Filter;
 }
 
-FCustomClassFilterData::FCustomClassFilterData(UClass*InClass, TSharedPtr<FFilterCategory> InCategory, FLinearColor InColor)
+FCustomClassFilterData::FCustomClassFilterData(UClass* InClass, TSharedPtr<FFilterCategory> InCategory, FLinearColor InColor)
 	: Class(InClass)
 	, Color(InColor)
 {
