Update BP Filter in Content Browser to include Aliased Compiled BPs

diff --git a/Engine/Source/Editor/EditorWidgets/Public/Filters/SAssetFilterBar.h b/Engine/Source/Editor/EditorWidgets/Public/Filters/SAssetFilterBar.h
--- a/Engine/Source/Editor/EditorWidgets/Public/Filters/SAssetFilterBar.h
+++ b/Engine/Source/Editor/EditorWidgets/Public/Filters/SAssetFilterBar.h
@@ -575,6 +575,14 @@ public:
 			}
 		}
 
+		// HACK: A blueprint can be shown as Blueprint or as BlueprintGeneratedClass, but we don't want to distinguish them while filtering.
+		// This should be removed, once all blueprints are shown as BlueprintGeneratedClass.
+		// Note: Adding the BlueprintGeneratedClass should occur before the below bRecursiveClasses check otherwise it will be added to the ExclusionSet
+		if (CombinedFilter.ClassPaths.Contains(FTopLevelAssetPath(TEXT("/Script/Engine"), TEXT("Blueprint"))))
+		{
+			CombinedFilter.ClassPaths.AddUnique(FTopLevelAssetPath(TEXT("/Script/Engine"), TEXT("BlueprintGeneratedClass")));
+		}
+
 		if ( CombinedFilter.bRecursiveClasses )
 		{
 			// Add exclusions for AssetTypeActions NOT in the filter.
@@ -590,13 +598,6 @@ public:
 			}
 		}
 
-		// HACK: A blueprint can be shown as Blueprint or as BlueprintGeneratedClass, but we don't want to distinguish them while filtering.
-		// This should be removed, once all blueprints are shown as BlueprintGeneratedClass.
-		if(CombinedFilter.ClassPaths.Contains(FTopLevelAssetPath(TEXT("/Script/Engine"), TEXT("Blueprint"))))
-		{
-			CombinedFilter.ClassPaths.AddUnique(FTopLevelAssetPath(TEXT("/Script/Engine"), TEXT("BlueprintGeneratedClass")));
-		}
-
 		return CombinedFilter;
 	}
 
