Added cache functionality to SDropTarget for very expensive AllowDrop operations.
Turned off by default.
By default, OnAllowDrop is called every frame, which is user bound. Setting "bUseAllowDropCache" to true will cache off AllowDrop the first time it is called.
The cache is cleared whenever a drag & drop operation starts.

diff --git a/Engine/Plugins/FX/Niagara/Source/NiagaraEditorWidgets/Private/Stack/SNiagaraStackFunctionInputValue.cpp b/Engine/Plugins/FX/Niagara/Source/NiagaraEditorWidgets/Private/Stack/SNiagaraStackFunctionInputValue.cpp
--- a/Engine/Plugins/FX/Niagara/Source/NiagaraEditorWidgets/Private/Stack/SNiagaraStackFunctionInputValue.cpp
+++ b/Engine/Plugins/FX/Niagara/Source/NiagaraEditorWidgets/Private/Stack/SNiagaraStackFunctionInputValue.cpp
@@ -86,6 +86,7 @@ void SNiagaraStackFunctionInputValue::Construct(const FArguments& InArgs, UNiaga
 			.HorizontalImage(FNiagaraEditorWidgetsStyle::Get().GetBrush("NiagaraEditor.Stack.DropTarget.BorderHorizontal"))
 			.VerticalImage(FNiagaraEditorWidgetsStyle::Get().GetBrush("NiagaraEditor.Stack.DropTarget.BorderVertical"))
 			.IsEnabled_UObject(FunctionInput, &UNiagaraStackEntry::GetOwnerIsEnabled)
+			.bUseAllowDropCache(true)
 			.Content()
 			[
 				// Values
diff --git a/Engine/Source/Editor/EditorWidgets/Private/SDropTarget.cpp b/Engine/Source/Editor/EditorWidgets/Private/SDropTarget.cpp
--- a/Engine/Source/Editor/EditorWidgets/Private/SDropTarget.cpp
+++ b/Engine/Source/Editor/EditorWidgets/Private/SDropTarget.cpp
@@ -28,6 +28,7 @@ void SDropTarget::Construct(const FArguments& InArgs)
 	OnDragEnterEvent = InArgs._OnDragEnter;
 	OnDragLeaveEvent = InArgs._OnDragLeave;
 	bOnlyRecognizeOnDragEnter = InArgs._bOnlyRecognizeOnDragEnter;
+	bUseAllowDropCache = InArgs._bUseAllowDropCache;
 	
 	bIsDragEventRecognized = false;
 	bAllowDrop = false;
@@ -39,6 +40,9 @@ void SDropTarget::Construct(const FArguments& InArgs)
 	VerticalImage = InArgs._VerticalImage;
 	HorizontalImage = InArgs._HorizontalImage;
 
+	// if we want to use the cache, we need to detect whether to clear the cache on tick
+	SetCanTick(bUseAllowDropCache);
+	
 	ChildSlot
 	[
 		SNew(SOverlay)
@@ -94,7 +98,19 @@ FReply SDropTarget::OnDragOver(const FGeometry& MyGeometry, const FDragDropEvent
 
 bool SDropTarget::AllowDrop(TSharedPtr<FDragDropOperation> DragDropOperation) const
 {
-	bAllowDrop = OnAllowDrop(DragDropOperation);
+	if(bUseAllowDropCache && AllowDropCache.IsSet())
+	{
+		bAllowDrop = AllowDropCache.GetValue();
+	}
+	else
+	{
+		bAllowDrop = OnAllowDrop(DragDropOperation);
+
+		if(bUseAllowDropCache)
+		{
+			AllowDropCache = bAllowDrop;
+		}
+	}
 	bIsDragEventRecognized = OnIsRecognized(DragDropOperation) || bAllowDrop;
 	return bAllowDrop;
 }
@@ -119,6 +135,20 @@ bool SDropTarget::OnIsRecognized(TSharedPtr<FDragDropOperation> DragDropOperatio
 	return false;
 }
 
+void SDropTarget::Tick(const FGeometry& AllottedGeometry, const double InCurrentTime, const float InDeltaTime)
+{
+	if(bUseAllowDropCache)
+	{
+		bWasDragDroppingLastFrame = bIsDragDropping;
+		bIsDragDropping = FSlateApplication::Get().IsDragDropping();
+
+		if(bIsDragDropping && !bWasDragDroppingLastFrame)
+		{
+			ClearAllowDropCache();
+		}
+	}
+}
+
 FReply SDropTarget::OnDrop(const FGeometry& MyGeometry, const FDragDropEvent& DragDropEvent)
 {
 	const bool bCurrentbAllowDrop = bAllowDrop;
diff --git a/Engine/Source/Editor/EditorWidgets/Public/SDropTarget.h b/Engine/Source/Editor/EditorWidgets/Public/SDropTarget.h
--- a/Engine/Source/Editor/EditorWidgets/Public/SDropTarget.h
+++ b/Engine/Source/Editor/EditorWidgets/Public/SDropTarget.h
@@ -43,6 +43,7 @@ public:
 		, _HorizontalImage(FAppStyle::GetBrush("WideDash.Horizontal"))
 		, _BackgroundImage(FAppStyle::GetBrush("DropTarget.Background"))
 		, _bOnlyRecognizeOnDragEnter(false)
+		, _bUseAllowDropCache(false)
 	{ }
 	
 		UE_DEPRECATED(5.0, "BackgroundColor has been removed. You may alter the background brush to get the same effect.")
@@ -78,6 +79,8 @@ public:
 		SLATE_EVENT(FOnDragAction, OnDragLeave)
 		/** When this is true, the drop target will only get recognized when entering while drag & dropping. */
 		SLATE_ATTRIBUTE(bool, bOnlyRecognizeOnDragEnter)
+		/** Whether to cache off the results of AllowDrop. Useful when then OnAllowDrop callback is expensive since it's called per frame. */
+		SLATE_ARGUMENT(bool, bUseAllowDropCache)
 
 		FOnDrop ConvertOnDropFn(const FOnDropDeprecated& LegacyDelegate)
 		{
@@ -96,13 +99,19 @@ public:
 
 	void Construct(const FArguments& InArgs );
 
+	void ClearAllowDropCache()
+	{
+		AllowDropCache.Reset();
+	}
+	
 protected:
 
 	bool AllowDrop(TSharedPtr<FDragDropOperation> DragDropOperation) const;
 
 	virtual bool OnAllowDrop(TSharedPtr<FDragDropOperation> DragDropOperation) const;
 	virtual bool OnIsRecognized(TSharedPtr<FDragDropOperation> DragDropOperation) const;
-
+	
+	virtual void Tick(const FGeometry& AllottedGeometry, const double InCurrentTime, const float InDeltaTime) override;
 protected:
 	// SWidget interface
 	virtual FReply OnDragOver( const FGeometry& MyGeometry, const FDragDropEvent& DragDropEvent ) override;
@@ -132,7 +141,8 @@ private:
 	FOnDragAction OnDragLeaveEvent;
 	/** Attribute to check if the drop target should only be useable when actually dragging over it. */
 	TAttribute<bool> bOnlyRecognizeOnDragEnter;
-
+	bool bUseAllowDropCache = false;
+	
 	/** The color of the vertical/horizontal images when the drop data is valid */
 	FSlateColor ValidColor;
 	/** The color of the vertical/horizontal images when the drop data is not valid */
@@ -148,4 +158,8 @@ private:
 	mutable bool bAllowDrop;
 	/** Is the drag operation currently over our airspace? */
 	mutable bool bIsDragOver;
+
+	mutable TOptional<bool> AllowDropCache;
+	bool bIsDragDropping = false;
+	bool bWasDragDroppingLastFrame = false;
 };
