Moving SinglePropertyView and DetailsView to EditorWidgets

diff --git a/Engine/Config/BaseEngine.ini b/Engine/Config/BaseEngine.ini
--- a/Engine/Config/BaseEngine.ini
+++ b/Engine/Config/BaseEngine.ini
@@ -1141,6 +1141,9 @@ DefaultWorldPartitionHLODLayer=/Engine/Maps/Templates/HLODs/HLODLayer_Instanced.
 
 +StructRedirects=(OldName="/Script/UMG.FieldNotificationId", NewName="/Script/FieldNotification.FieldNotificationId")
 +ClassRedirects=(OldName="/Script/UMG.NotifyFieldValueChanged", NewName="/Script/FieldNotification.NotifyFieldValueChanged")
++ClassRedirects=(OldName="/Script/UMGEditor.DetailsView", NewName="/Script/EditorWidgets.DetailsView")
++ClassRedirects=(OldName="/Script/UMGEditor.SinglePropertyView", NewName="/Script/EditorWidgets.SinglePropertyView")
++ClassRedirects=(OldName="/Script/UMGEditor.PropertyViewBase", NewName="/Script/EditorWidgets.PropertyViewBase")
 
 ; Refactoring most of HeadMountedDisplay engine module into XRBase plugin.
 +ClassRedirects=(OldName="HeadMountedDisplayFunctionLibrary", NewName="/Script/XRBase.HeadMountedDisplayFunctionLibrary")
diff --git a/Engine/Source/Editor/EditorWidgets/EditorWidgets.Build.cs b/Engine/Source/Editor/EditorWidgets/EditorWidgets.Build.cs
--- a/Engine/Source/Editor/EditorWidgets/EditorWidgets.Build.cs
+++ b/Engine/Source/Editor/EditorWidgets/EditorWidgets.Build.cs
@@ -21,6 +21,7 @@ public class EditorWidgets : ModuleRules
 				"ToolWidgets",
 				"EditorConfig",
 				"AssetDefinition",
+				"UMG",
 			}
 		);
 		
diff --git a/Engine/Source/Editor/EditorWidgets/Private/Components/DetailsView.cpp b/Engine/Source/Editor/EditorWidgets/Private/Components/DetailsView.cpp
new file mode 100644
--- /dev/null
+++ b/Engine/Source/Editor/EditorWidgets/Private/Components/DetailsView.cpp
@@ -0,0 +1,239 @@
+// Copyright Epic Games, Inc. All Rights Reserved.
+
+#include "Components/DetailsView.h"
+
+#include "Components/PropertyViewHelper.h"
+#include "Containers/UnrealString.h"
+#include "CoreGlobals.h"
+#include "Delegates/Delegate.h"
+#include "DetailsViewArgs.h"
+#include "HAL/Platform.h"
+#include "IDetailsView.h"
+#include "Internationalization/Internationalization.h"
+#include "Internationalization/Text.h"
+#include "Misc/AssertionMacros.h"
+#include "Misc/Attribute.h"
+#include "Modules/ModuleManager.h"
+#include "ObjectEditorUtils.h"
+#include "PropertyEditorDelegates.h"
+#include "PropertyEditorModule.h"
+#include "UObject/SoftObjectPtr.h"
+#include "UObject/UnrealType.h"
+#include "Widgets/DeclarativeSyntaxSupport.h"
+#include "Widgets/Layout/SBorder.h"
+#include "Widgets/Text/STextBlock.h"
+
+class UObject;
+
+#define LOCTEXT_NAMESPACE "UMG"
+
+/////////////////////////////////////////////////////
+// UDetailsView
+
+
+void UDetailsView::ReleaseSlateResources(bool bReleaseChildren)
+{
+	Super::ReleaseSlateResources(bReleaseChildren);
+
+	DetailViewWidget.Reset();
+}
+
+
+void UDetailsView::BuildContentWidget()
+{
+	DetailViewWidget.Reset();
+
+	if (!GetDisplayWidget().IsValid())
+	{
+		return;
+	}
+
+	bool bCreateMissingWidget = true;
+	FText MissingWidgetText = FPropertyViewHelper::EditorOnlyText;
+
+	if (GIsEditor)
+	{
+		UObject* ViewedObject = GetObject();
+		if (ViewedObject == nullptr)
+		{
+			bool bIsObjectNull = Object.IsNull();
+			if (bIsObjectNull)
+			{
+				MissingWidgetText = FPropertyViewHelper::UndefinedObjectText;
+			}
+			else
+			{
+				MissingWidgetText = FPropertyViewHelper::UnloadedObjectText;
+			}
+		}
+		else
+		{
+			FPropertyEditorModule& PropertyEditorModule = FModuleManager::GetModuleChecked<FPropertyEditorModule>("PropertyEditor");
+
+			FDetailsViewArgs DetailsViewArgs;
+			DetailsViewArgs.NameAreaSettings = FDetailsViewArgs::HideNameArea;
+			DetailsViewArgs.bUpdatesFromSelection = false;
+			DetailsViewArgs.bLockable = false;
+			DetailsViewArgs.bShowPropertyMatrixButton = false;
+			DetailsViewArgs.NotifyHook = this;
+
+			DetailsViewArgs.ViewIdentifier = ViewIdentifier;
+			DetailsViewArgs.bAllowSearch = bAllowFiltering;
+			DetailsViewArgs.bAllowFavoriteSystem = bAllowFavoriteSystem;
+			DetailsViewArgs.bShowOptions = bAllowFiltering;
+			DetailsViewArgs.bShowModifiedPropertiesOption = bShowModifiedPropertiesOption;
+			DetailsViewArgs.bShowKeyablePropertiesOption = bShowKeyablePropertiesOption;
+			DetailsViewArgs.bShowAnimatedPropertiesOption = bShowAnimatedPropertiesOption;
+			DetailsViewArgs.bShowScrollBar = bShowScrollBar;
+			DetailsViewArgs.bForceHiddenPropertyVisibility = bForceHiddenPropertyVisibility;
+			DetailsViewArgs.ColumnWidth = ColumnWidth;
+			DetailsViewArgs.bShowCustomFilterOption = PropertiesToShow.Num() != 0 || CategoriesToShow.Num() != 0;
+			
+			DetailViewWidget = PropertyEditorModule.CreateDetailView(DetailsViewArgs);
+			
+			DetailViewWidget->SetCustomFilterLabel(LOCTEXT("ShowAllParameters", "Show All Parameters"));
+			DetailViewWidget->SetCustomFilterDelegate(FSimpleDelegate::CreateUObject(this, &UDetailsView::ToggleShowingOnlyAllowedProperties));
+			DetailViewWidget->SetIsPropertyVisibleDelegate(FIsPropertyVisible::CreateUObject(this, &UDetailsView::GetIsPropertyVisible));
+			DetailViewWidget->SetIsCustomRowVisibleDelegate(FIsCustomRowVisible::CreateUObject(this, &UDetailsView::GetIsRowVisible));
+			DetailViewWidget->SetObject(ViewedObject);
+			if (DetailViewWidget.IsValid())
+			{
+				GetDisplayWidget()->SetContent(DetailViewWidget.ToSharedRef());
+				bCreateMissingWidget = false;
+			}
+			else
+			{
+				MissingWidgetText = FPropertyViewHelper::UnknownErrorText;
+			}
+		}
+	}
+
+	if (bCreateMissingWidget)
+	{
+		GetDisplayWidget()->SetContent(
+			SNew(STextBlock)
+			.Text(MissingWidgetText)
+		);
+	}
+}
+
+
+void UDetailsView::OnObjectChanged()
+{
+	UObject* ViewedObject = GetObject();
+	if (DetailViewWidget.IsValid() && ViewedObject != nullptr)
+	{
+		DetailViewWidget->SetObject(ViewedObject);
+	}
+	else
+	{
+		BuildContentWidget();
+	}
+}
+
+
+void UDetailsView::NotifyPostChange(const FPropertyChangedEvent& PropertyChangedEvent, FProperty* PropertyThatChanged)
+{
+	FNotifyHook::NotifyPostChange(PropertyChangedEvent, PropertyThatChanged);
+
+	FName PropertyName = PropertyThatChanged ? PropertyThatChanged->GetFName() : NAME_None;
+	OnPropertyChangedBroadcast(PropertyName);
+}
+
+
+void UDetailsView::PostEditChangeProperty(FPropertyChangedEvent& PropertyChangedEvent)
+{
+	Super::PostEditChangeProperty(PropertyChangedEvent);
+	
+	if (IsDesignTime())
+	{
+		if (PropertyChangedEvent.GetPropertyName() == GET_MEMBER_NAME_CHECKED(UDetailsView, ViewIdentifier)
+			|| PropertyChangedEvent.GetPropertyName() == GET_MEMBER_NAME_CHECKED(UDetailsView, bAllowFiltering)
+			|| PropertyChangedEvent.GetPropertyName() == GET_MEMBER_NAME_CHECKED(UDetailsView, bAllowFavoriteSystem)
+			|| PropertyChangedEvent.GetPropertyName() == GET_MEMBER_NAME_CHECKED(UDetailsView, bShowModifiedPropertiesOption)
+			|| PropertyChangedEvent.GetPropertyName() == GET_MEMBER_NAME_CHECKED(UDetailsView, bShowKeyablePropertiesOption)
+			|| PropertyChangedEvent.GetPropertyName() == GET_MEMBER_NAME_CHECKED(UDetailsView, bShowAnimatedPropertiesOption)
+			|| PropertyChangedEvent.GetPropertyName() == GET_MEMBER_NAME_CHECKED(UDetailsView, bShowScrollBar)
+			|| PropertyChangedEvent.GetPropertyName() == GET_MEMBER_NAME_CHECKED(UDetailsView, bForceHiddenPropertyVisibility)
+			|| PropertyChangedEvent.GetPropertyName() == GET_MEMBER_NAME_CHECKED(UDetailsView, ColumnWidth))
+		{
+			BuildContentWidget();
+		}
+	}
+}
+
+
+void UDetailsView::ToggleShowingOnlyAllowedProperties()
+{
+	bShowOnlyAllowed = !bShowOnlyAllowed;
+	if (DetailViewWidget.IsValid())
+	{
+		DetailViewWidget->ForceRefresh();
+	}
+}
+
+bool UDetailsView::IsRowVisibilityFiltered() const
+{
+	return bShowOnlyAllowed && (PropertiesToShow.Num() > 0 || CategoriesToShow.Num() > 0);
+}
+
+bool UDetailsView::GetIsPropertyVisible(const FPropertyAndParent& PropertyAndParent) const
+{
+	if (!IsRowVisibilityFiltered())
+	{
+		return true;
+	}
+	if (PropertiesToShow.Contains(PropertyAndParent.Property.GetFName()))
+	{
+		return true;
+	}
+
+	const FProperty* Property = PropertyAndParent.ParentProperties.Num() > 0 ? PropertyAndParent.ParentProperties.Last() : &PropertyAndParent.Property;
+
+	// Check against the parent property name
+	if (PropertiesToShow.Contains(Property->GetFName()))
+	{
+		return true;
+	}
+
+	// get the topmost parent's category name if the property has one
+	const FString CategoryString = FObjectEditorUtils::GetCategoryFName(Property).ToString();
+	FString SubcategoryString = CategoryString;
+	
+	int32 SubcategoryStart;
+	if (CategoryString.FindChar(TEXT('|'), SubcategoryStart))
+	{
+		SubcategoryString = CategoryString.Left(SubcategoryStart);
+	}
+
+	if (CategoriesToShow.ContainsByPredicate([&CategoryString, &SubcategoryString](const FName& Element)
+		{
+			FString ElementString = Element.ToString();
+			return CategoryString == ElementString || SubcategoryString == ElementString;
+		}))
+	{
+		return true;
+	}
+	return false;
+}
+
+bool UDetailsView::GetIsRowVisible(FName InRowName, FName InParentName) const
+{
+	if (!IsRowVisibilityFiltered())
+	{
+		return true;
+	}
+	if (PropertiesToShow.Contains(InRowName))
+	{
+		return true;
+	}
+	if (CategoriesToShow.Contains(InParentName))
+	{
+		return true;
+	}
+	return false;
+}
+
+/////////////////////////////////////////////////////
+
+#undef LOCTEXT_NAMESPACE
diff --git a/Engine/Source/Editor/EditorWidgets/Private/Components/PropertyViewBase.cpp b/Engine/Source/Editor/EditorWidgets/Private/Components/PropertyViewBase.cpp
new file mode 100644
--- /dev/null
+++ b/Engine/Source/Editor/EditorWidgets/Private/Components/PropertyViewBase.cpp
@@ -0,0 +1,96 @@
+// Copyright Epic Games, Inc. All Rights Reserved.
+
+#include "Components/PropertyViewBase.h"
+
+#include "Async/Async.h"
+#include "Editor.h"
+#include "Engine/World.h"
+#include "PropertyEditorModule.h"
+#include "Widgets/Layout/SBorder.h"
+#include "UObject/UObjectGlobals.h"
+
+#define LOCTEXT_NAMESPACE "UMG"
+
+/////////////////////////////////////////////////////
+// UPropertyViewBase
+
+
+void UPropertyViewBase::ReleaseSlateResources(bool bReleaseChildren)
+{
+	Super::ReleaseSlateResources(bReleaseChildren);
+	DisplayedWidget.Reset();
+}
+
+
+TSharedRef<SWidget> UPropertyViewBase::RebuildWidget()
+{
+	DisplayedWidget = SNew(SBorder)
+		.Padding(0.0f)
+		.HAlign(HAlign_Fill)
+		.VAlign(VAlign_Fill)
+		.BorderImage(FAppStyle::GetBrush("NoBorder"));
+	
+	BuildContentWidget();
+
+	return DisplayedWidget.ToSharedRef();
+}
+
+
+UObject* UPropertyViewBase::GetObject() const
+{
+	return Object.Get();
+}
+
+
+void UPropertyViewBase::SetObject(UObject* InObject)
+{
+	if (Object.Get() != InObject)
+	{
+		Object = InObject;
+		OnObjectChanged();
+	}
+}
+
+
+void UPropertyViewBase::OnPropertyChangedBroadcast(FName PropertyName)
+{
+	OnPropertyChanged.Broadcast(PropertyName);
+}
+
+
+void UPropertyViewBase::PostLoad()
+{
+	Super::PostLoad();
+
+	if (Object.IsNull() && !SoftObjectPath_DEPRECATED.IsNull())
+	{
+		Object = SoftObjectPath_DEPRECATED;
+	}
+
+	if(!Object.IsValid() && Object.ToSoftObjectPath().IsAsset() && bAutoLoadAsset && !HasAnyFlags(RF_BeginDestroyed))
+	{
+		Object.LoadSynchronous();
+		BuildContentWidget();
+	}
+}
+
+
+void UPropertyViewBase::PostEditChangeProperty(FPropertyChangedEvent& PropertyChangedEvent)
+{
+	if (PropertyChangedEvent.GetPropertyName() == GET_MEMBER_NAME_CHECKED(UPropertyViewBase, Object))
+	{
+		OnObjectChanged();
+	}
+
+	Super::PostEditChangeProperty(PropertyChangedEvent);
+}
+
+
+const FText UPropertyViewBase::GetPaletteCategory()
+{
+	return LOCTEXT("Editor", "Editor");
+}
+
+/////////////////////////////////////////////////////
+
+#undef LOCTEXT_NAMESPACE
diff --git a/Engine/Source/Editor/EditorWidgets/Private/Components/PropertyViewHelper.cpp b/Engine/Source/Editor/EditorWidgets/Private/Components/PropertyViewHelper.cpp
new file mode 100644
--- /dev/null
+++ b/Engine/Source/Editor/EditorWidgets/Private/Components/PropertyViewHelper.cpp
@@ -0,0 +1,25 @@
+// Copyright Epic Games, Inc. All Rights Reserved.
+
+#include "Components/PropertyViewHelper.h"
+
+#include "Internationalization/Internationalization.h"
+#include "Internationalization/Text.h"
+
+
+#define LOCTEXT_NAMESPACE "UMG"
+
+/////////////////////////////////////////////////////
+// FPropertyViewHelper
+
+const FText FPropertyViewHelper::UndefinedObjectText = LOCTEXT("UndefinedObject", "Undefined Object");
+const FText FPropertyViewHelper::UnloadedObjectText = LOCTEXT("UnloadedObject", "Unloaded Object");
+const FText FPropertyViewHelper::UnknownErrorText = LOCTEXT("UnknownError", "Unknown Error");
+const FText FPropertyViewHelper::EditorOnlyText = LOCTEXT("EditorWidget", "Editor Only");
+const FText FPropertyViewHelper::UndefinedPropertyText = LOCTEXT("UndefinedProperty", "Undefined Property");
+const FText FPropertyViewHelper::UnknownPropertyText = LOCTEXT("UnknownProperty", "Unknown Property");
+const FText FPropertyViewHelper::InvalidPropertyText = LOCTEXT("InvalidPropertyName", "Invalid Property");
+const FText FPropertyViewHelper::UnsupportedPropertyText = LOCTEXT("SinglePropertyView_UnsupporteddProperty", "Unsupported Property");
+
+/////////////////////////////////////////////////////
+
+#undef LOCTEXT_NAMESPACE
\ No newline at end of file
diff --git a/Engine/Source/Editor/EditorWidgets/Private/Components/PropertyViewHelper.h b/Engine/Source/Editor/EditorWidgets/Private/Components/PropertyViewHelper.h
new file mode 100644
--- /dev/null
+++ b/Engine/Source/Editor/EditorWidgets/Private/Components/PropertyViewHelper.h
@@ -0,0 +1,17 @@
+// Copyright Epic Games, Inc. All Rights Reserved.
+
+#pragma once
+
+class FText;
+
+struct FPropertyViewHelper
+{
+	static const FText UndefinedObjectText;
+	static const FText UnloadedObjectText;
+	static const FText UnknownErrorText;
+	static const FText EditorOnlyText;
+	static const FText UndefinedPropertyText;
+	static const FText UnknownPropertyText;
+	static const FText InvalidPropertyText;
+	static const FText UnsupportedPropertyText;
+};
diff --git a/Engine/Source/Editor/EditorWidgets/Private/Components/SinglePropertyView.cpp b/Engine/Source/Editor/EditorWidgets/Private/Components/SinglePropertyView.cpp
new file mode 100644
--- /dev/null
+++ b/Engine/Source/Editor/EditorWidgets/Private/Components/SinglePropertyView.cpp
@@ -0,0 +1,183 @@
+// Copyright Epic Games, Inc. All Rights Reserved.
+
+#include "Components/SinglePropertyView.h"
+
+#include "Components/PropertyViewHelper.h"
+#include "CoreGlobals.h"
+#include "Delegates/Delegate.h"
+#include "ISinglePropertyView.h"
+#include "Misc/AssertionMacros.h"
+#include "Misc/Attribute.h"
+#include "Modules/ModuleManager.h"
+#include "PropertyEditorModule.h"
+#include "UObject/Class.h"
+#include "UObject/Field.h"
+#include "UObject/Object.h"
+#include "UObject/SoftObjectPtr.h"
+#include "UObject/UnrealNames.h"
+#include "UObject/UnrealType.h"
+#include "Widgets/DeclarativeSyntaxSupport.h"
+#include "Widgets/Layout/SBorder.h"
+#include "Widgets/Text/STextBlock.h"
+
+
+#define LOCTEXT_NAMESPACE "UMG"
+
+/////////////////////////////////////////////////////
+// USinglePropertyView
+
+
+void USinglePropertyView::ReleaseSlateResources(bool bReleaseChildren)
+{
+	Super::ReleaseSlateResources(bReleaseChildren);
+
+	SinglePropertyViewWidget.Reset();
+}
+
+
+void USinglePropertyView::BuildContentWidget()
+{
+	SinglePropertyViewWidget.Reset();
+
+	if (!GetDisplayWidget().IsValid())
+	{
+		return;
+	}
+
+	bool bCreateMissingWidget = true;
+	FText MissingWidgetText = FPropertyViewHelper::EditorOnlyText;
+
+	if (GIsEditor)
+	{
+		UObject* ViewedObject = GetObject();
+		if (ViewedObject == nullptr)
+		{
+			bool bIsObjectNull = Object.IsNull();
+			if (bIsObjectNull)
+			{
+				MissingWidgetText = FPropertyViewHelper::UndefinedObjectText;
+			}
+			else
+			{
+				MissingWidgetText = FPropertyViewHelper::UnloadedObjectText;
+			}
+		}
+		else if (PropertyName == NAME_None)
+		{
+			MissingWidgetText = FPropertyViewHelper::UndefinedPropertyText;
+		}
+		else
+		{
+			FProperty* Property = ViewedObject->GetClass()->FindPropertyByName(PropertyName);
+			if (Property == nullptr)
+			{
+				MissingWidgetText = FPropertyViewHelper::UnknownPropertyText;
+			}
+			else if (!Property->HasAllPropertyFlags(CPF_Edit))
+			{
+				MissingWidgetText = FPropertyViewHelper::InvalidPropertyText;
+			}
+			else if (CastField<FArrayProperty>(Property) || CastField<FMapProperty>(Property) || CastField<FSetProperty>(Property))
+			{
+				MissingWidgetText = FPropertyViewHelper::UnsupportedPropertyText;
+			}
+			else
+			{
+				FPropertyEditorModule& PropertyEditorModule = FModuleManager::GetModuleChecked<FPropertyEditorModule>("PropertyEditor");
+				FSinglePropertyParams SinglePropertyArgs;
+				SinglePropertyArgs.NameOverride = NameOverride;
+				SinglePropertyViewWidget = PropertyEditorModule.CreateSingleProperty(ViewedObject, PropertyName, SinglePropertyArgs);
+
+				if (SinglePropertyViewWidget.IsValid())
+				{
+					FSimpleDelegate PropertyChanged = FSimpleDelegate::CreateUObject(this, &USinglePropertyView::InternalSinglePropertyChanged);
+					SinglePropertyViewWidget->SetOnPropertyValueChanged(PropertyChanged);
+
+					GetDisplayWidget()->SetContent(SinglePropertyViewWidget.ToSharedRef());
+					bCreateMissingWidget = false;
+				}
+				else
+				{
+					// Some built-in structs like FColor are supported, others aren't
+					if (CastField<FStructProperty>(Property))
+					{
+						MissingWidgetText = FPropertyViewHelper::UnsupportedPropertyText;
+					}
+					else
+					{
+						MissingWidgetText = FPropertyViewHelper::UnknownErrorText;
+					}
+				}
+			}
+		}
+	}
+
+	if (bCreateMissingWidget)
+	{
+		GetDisplayWidget()->SetContent(
+			SNew(STextBlock)
+			.Text(MissingWidgetText)
+		);
+	}
+}
+
+
+FName USinglePropertyView::GetPropertyName() const
+{
+	return PropertyName;
+}
+
+
+void USinglePropertyView::SetPropertyName(FName InPropertyName)
+{
+	if (PropertyName != InPropertyName)
+	{
+		PropertyName = InPropertyName;
+		BuildContentWidget();
+	}
+}
+
+FText USinglePropertyView::GetNameOverride() const
+{
+	return NameOverride;
+}
+
+
+void USinglePropertyView::SetNameOverride(FText InNameOverride)
+{
+	if (!NameOverride.EqualTo(InNameOverride))
+	{
+		NameOverride = InNameOverride;
+		BuildContentWidget();
+	}
+}
+
+void USinglePropertyView::OnObjectChanged()
+{
+	BuildContentWidget();
+}
+
+
+void USinglePropertyView::InternalSinglePropertyChanged()
+{
+	OnPropertyChangedBroadcast(GetPropertyName());
+}
+
+
+void USinglePropertyView::PostEditChangeProperty(FPropertyChangedEvent& PropertyChangedEvent)
+{
+	Super::PostEditChangeProperty(PropertyChangedEvent);
+
+	if (IsDesignTime())
+	{
+		if (PropertyChangedEvent.GetPropertyName() == GET_MEMBER_NAME_CHECKED(USinglePropertyView, PropertyName)
+			|| PropertyChangedEvent.GetPropertyName() == GET_MEMBER_NAME_CHECKED(USinglePropertyView, NameOverride))
+		{
+			BuildContentWidget();
+		}
+	}
+}
+
+/////////////////////////////////////////////////////
+
+#undef LOCTEXT_NAMESPACE
diff --git a/Engine/Source/Editor/EditorWidgets/Public/Components/DetailsView.h b/Engine/Source/Editor/EditorWidgets/Public/Components/DetailsView.h
new file mode 100644
--- /dev/null
+++ b/Engine/Source/Editor/EditorWidgets/Public/Components/DetailsView.h
@@ -0,0 +1,106 @@
+// Copyright Epic Games, Inc. All Rights Reserved.
+
+#pragma once
+
+#include "Components/PropertyViewBase.h"
+#include "Containers/Array.h"
+#include "Misc/NotifyHook.h"
+#include "PropertyEditorDelegates.h"
+#include "Templates/SharedPointer.h"
+#include "UObject/NameTypes.h"
+#include "UObject/UObjectGlobals.h"
+#include "UObject/UnrealNames.h"
+
+#include "DetailsView.generated.h"
+
+class FProperty;
+class IDetailsView;
+class UObject;
+struct FPropertyAndParent;
+struct FPropertyChangedEvent;
+
+/**
+ * The details view allows you to display the value of an object properties.
+ */
+UCLASS()
+class EDITORWIDGETS_API UDetailsView : public UPropertyViewBase, public FNotifyHook
+{
+	GENERATED_BODY()
+
+public:
+	/** True if we allow filtering through search and the filter dropdown menu. */
+	UPROPERTY(EditAnywhere, Category = "View")
+	bool bAllowFiltering = false;
+
+	/** If false, the current properties editor will never display the favorite system */
+	UPROPERTY(EditAnywhere, AdvancedDisplay, Category = "View")
+	bool bAllowFavoriteSystem = false;
+
+	/** True if you want to show the 'Show Only Modified Properties'. Only valid in conjunction with bAllowFiltering */
+	UPROPERTY(EditAnywhere, AdvancedDisplay, Category = "View", meta = (EditCondition = "bAllowFiltering"))
+	bool bShowModifiedPropertiesOption = false;
+
+	/** True if you want to show the 'Show Only Keyable Properties'. Only valid in conjunction with bAllowFiltering */
+	UPROPERTY(EditAnywhere, AdvancedDisplay, Category = "View", meta = (EditCondition = "bAllowFiltering"))
+	bool bShowKeyablePropertiesOption = true;
+
+	/** True if you want to show the 'Show Only Animated Properties'. Only valid in conjunction with bAllowFiltering */
+	UPROPERTY(EditAnywhere, AdvancedDisplay, Category = "View", meta = (EditCondition = "bAllowFiltering"))
+	bool bShowAnimatedPropertiesOption = true;
+
+	/** The default column width */
+	UPROPERTY(EditAnywhere, Category = "View")
+	float ColumnWidth = 0.65f;
+
+	/** If false, the details panel's scrollbar will always be hidden. Useful when embedding details panels in widgets that either grow to accommodate them, or with scrollbars of their own. */
+	UPROPERTY(EditAnywhere, AdvancedDisplay, Category = "View")
+	bool bShowScrollBar = true;
+
+	/** If true, all properties will be visible, not just those with CPF_Edit */
+	UPROPERTY(EditAnywhere, AdvancedDisplay, Category = "View")
+	bool bForceHiddenPropertyVisibility = false;
+
+	/** Identifier for this details view; NAME_None if this view is anonymous */
+	UPROPERTY(EditAnywhere, AdvancedDisplay, Category = "View")
+	FName ViewIdentifier = NAME_None;
+
+	/** Which categories to show in the details panel. If both this and the Properties To Show lists are empty, all properties will show. */
+	UPROPERTY(EditAnywhere, Category = "View")
+	TArray<FName> CategoriesToShow;
+
+	/** Which properties to show in the details panel. If both this and the Categories To Show lists are empty, all properties will show. */
+	UPROPERTY(EditAnywhere, Category = "View")
+	TArray<FName> PropertiesToShow;
+
+	//~ UPropertyViewBase interface
+private:
+	virtual void BuildContentWidget() override;
+	virtual void OnObjectChanged() override;
+	//~ End of UPropertyViewBase interface
+
+public:
+	//~ UWidget interface
+	virtual void ReleaseSlateResources(bool bReleaseChildren) override;
+	// End of UWidget interface
+
+	//~ FNotifyHook interface
+	virtual void NotifyPostChange(const FPropertyChangedEvent& PropertyChangedEvent, FProperty* PropertyThatChanged) override;
+	//~ End of FNotifyHook interface
+
+	//~ UObject interface
+	virtual void PostEditChangeProperty(FPropertyChangedEvent& PropertyChangedEvent) override;
+	//~ End of UObject interface
+
+protected:
+	void ToggleShowingOnlyAllowedProperties();
+	virtual bool IsRowVisibilityFiltered() const;
+	virtual bool GetIsPropertyVisible(const FPropertyAndParent& PropertyAndParent) const;
+	virtual bool GetIsRowVisible(FName InRowName, FName InParentName) const;
+
+private:
+	TSharedPtr<IDetailsView> DetailViewWidget;
+
+	/** Showing properties in this details panel works by allowing only specific categories and properties. This flag enables you to show all properties without needing to specify. */
+	UPROPERTY()
+	bool bShowOnlyAllowed = true;
+};
diff --git a/Engine/Source/Editor/EditorWidgets/Public/Components/PropertyViewBase.h b/Engine/Source/Editor/EditorWidgets/Public/Components/PropertyViewBase.h
new file mode 100644
--- /dev/null
+++ b/Engine/Source/Editor/EditorWidgets/Public/Components/PropertyViewBase.h
@@ -0,0 +1,73 @@
+// Copyright Epic Games, Inc. All Rights Reserved.
+
+#pragma once
+
+#include "CoreMinimal.h"
+#include "Components/Widget.h"
+#include "Widgets/SWidget.h"
+#include "UObject/SoftObjectPath.h"
+#include "UObject/SoftObjectPtr.h"
+#include "PropertyViewBase.generated.h"
+
+class SBorder;
+
+/** Sets a delegate called when the property value changes */
+DECLARE_DYNAMIC_MULTICAST_DELEGATE_OneParam(FOnPropertyValueChanged, FName, PropertyName);
+
+
+/**
+ * Base of property view allows you to display the value of an object properties.
+ */
+UCLASS(Abstract)
+class EDITORWIDGETS_API UPropertyViewBase : public UWidget
+{
+	GENERATED_BODY()
+
+protected:
+	/** The object to view. */
+	UPROPERTY(meta = (DisplayName="Object"))
+	TSoftObjectPtr<UObject> Object;
+
+	UPROPERTY()
+	FSoftObjectPath SoftObjectPath_DEPRECATED;
+
+	/** Load the object (if it's an asset) when the widget is created. */
+	UPROPERTY(EditAnywhere, AdvancedDisplay, Category = "View")
+	bool bAutoLoadAsset = true;
+
+	/** Sets a delegate called when the property value changes */
+	UPROPERTY(BlueprintAssignable, Category = "View|Event")
+	FOnPropertyValueChanged OnPropertyChanged;
+
+
+public:
+	UFUNCTION(BlueprintCallable, Category = "View")
+	UObject* GetObject() const;
+
+	UFUNCTION(BlueprintCallable, Category = "View")
+	void SetObject(UObject* NewObject);
+
+protected:
+	virtual void BuildContentWidget() PURE_VIRTUAL(UPropertyViewBase::BuildContentWidget, );
+	virtual void OnObjectChanged() { }
+	TSharedPtr<SBorder> GetDisplayWidget() const { return DisplayedWidget; }
+	void OnPropertyChangedBroadcast(FName PropertyName);
+
+	//~ UWidget interface
+public:
+	virtual void ReleaseSlateResources(bool bReleaseChildren) override;
+	virtual const FText GetPaletteCategory() override;
+
+protected:
+	virtual TSharedRef<SWidget> RebuildWidget() override;
+	//~ End of UWidget interface
+
+	//~ UObject interface
+public:
+	virtual void PostLoad() override;
+	virtual void PostEditChangeProperty(FPropertyChangedEvent& PropertyChangedEvent) override;
+	//~ End of UObject interface
+
+private:
+	TSharedPtr<SBorder> DisplayedWidget;
+};
diff --git a/Engine/Source/Editor/EditorWidgets/Public/Components/SinglePropertyView.h b/Engine/Source/Editor/EditorWidgets/Public/Components/SinglePropertyView.h
new file mode 100644
--- /dev/null
+++ b/Engine/Source/Editor/EditorWidgets/Public/Components/SinglePropertyView.h
@@ -0,0 +1,69 @@
+// Copyright Epic Games, Inc. All Rights Reserved.
+
+#pragma once
+
+#include "Components/PropertyViewBase.h"
+#include "Internationalization/Text.h"
+#include "Templates/SharedPointer.h"
+#include "UObject/NameTypes.h"
+#include "UObject/UObjectGlobals.h"
+
+#include "SinglePropertyView.generated.h"
+
+class ISinglePropertyView;
+class UObject;
+struct FFrame;
+struct FPropertyChangedEvent;
+
+/**
+ * The single property view allows you to display the value of an object's property.
+ */
+UCLASS()
+class EDITORWIDGETS_API USinglePropertyView : public UPropertyViewBase
+{
+	GENERATED_BODY()
+
+private:
+	/** The name of the property to display. */
+	UPROPERTY(EditAnywhere, Category = "View")
+	FName PropertyName;
+
+	/** Override for the property name that will be displayed instead of the property name. */
+	UPROPERTY(EditAnywhere, AdvancedDisplay, Category = "View")
+	FText NameOverride;
+
+public:
+	UFUNCTION(BlueprintCallable, Category = "View")
+	FName GetPropertyName() const;
+
+	UFUNCTION(BlueprintCallable, Category = "View")
+	void SetPropertyName(FName NewPropertyName);
+
+	UFUNCTION(BlueprintCallable, Category = "View")
+	FText GetNameOverride() const;
+
+	UFUNCTION(BlueprintCallable, Category = "View")
+	void SetNameOverride(FText NewPropertyName);
+
+private:
+	void InternalSinglePropertyChanged();
+
+	//~ UPropertyViewBase interface
+private:
+	virtual void BuildContentWidget() override;
+	virtual void OnObjectChanged() override;
+	//~ End of UPropertyViewBase interface
+
+	//~ UWidget interface
+public:
+	virtual void ReleaseSlateResources(bool bReleaseChildren) override;
+	//~ End of UWidget interface
+
+	// UObject interface
+public:
+	virtual void PostEditChangeProperty(FPropertyChangedEvent& PropertyChangedEvent) override;
+	// End of UObject interface
+
+private:
+	TSharedPtr<ISinglePropertyView> SinglePropertyViewWidget;
+};
diff --git a/Engine/Source/Editor/UMGEditor/Private/Components/DetailsView.cpp b/Engine/Source/Editor/UMGEditor/Private/Components/DetailsView.cpp
deleted file mode 100644
diff --git a/Engine/Source/Editor/UMGEditor/Private/Components/PropertyViewBase.cpp b/Engine/Source/Editor/UMGEditor/Private/Components/PropertyViewBase.cpp
deleted file mode 100644
diff --git a/Engine/Source/Editor/UMGEditor/Private/Components/PropertyViewHelper.cpp b/Engine/Source/Editor/UMGEditor/Private/Components/PropertyViewHelper.cpp
deleted file mode 100644
diff --git a/Engine/Source/Editor/UMGEditor/Private/Components/PropertyViewHelper.h b/Engine/Source/Editor/UMGEditor/Private/Components/PropertyViewHelper.h
deleted file mode 100644
diff --git a/Engine/Source/Editor/UMGEditor/Private/Components/SinglePropertyView.cpp b/Engine/Source/Editor/UMGEditor/Private/Components/SinglePropertyView.cpp
deleted file mode 100644
diff --git a/Engine/Source/Editor/UMGEditor/Public/Components/DetailsView.h b/Engine/Source/Editor/UMGEditor/Public/Components/DetailsView.h
deleted file mode 100644
diff --git a/Engine/Source/Editor/UMGEditor/Public/Components/PropertyViewBase.h b/Engine/Source/Editor/UMGEditor/Public/Components/PropertyViewBase.h
deleted file mode 100644
diff --git a/Engine/Source/Editor/UMGEditor/Public/Components/SinglePropertyView.h b/Engine/Source/Editor/UMGEditor/Public/Components/SinglePropertyView.h
deleted file mode 100644
