Start showing the asset scan progress bar earlier so it doesn't jump to 50% when loading from the cache while still checking directories in the background

diff --git a/Engine/Source/Editor/EditorWidgets/Private/SAssetDiscoveryIndicator.cpp b/Engine/Source/Editor/EditorWidgets/Private/SAssetDiscoveryIndicator.cpp
--- a/Engine/Source/Editor/EditorWidgets/Private/SAssetDiscoveryIndicator.cpp
+++ b/Engine/Source/Editor/EditorWidgets/Private/SAssetDiscoveryIndicator.cpp
@@ -156,6 +156,13 @@ void SAssetDiscoveryIndicator::OnAssetRegistryFileLoadProgress(const IAssetRegis
 	{
 		// Marquee while we're discovering asset files as we can't yet show an accurate percentage
 		Progress = TOptional<float>();
+
+		// Once we have discovered enough assets, show a preview of the progress bar so it doesn't jump from 0 to 50%
+		if (ProgressUpdateData.NumAssetsProcessedByAssetRegistry > 10000 && ProgressUpdateData.NumTotalAssets > 0)
+		{
+			Progress = ProgressUpdateData.NumAssetsProcessedByAssetRegistry / (float)ProgressUpdateData.NumTotalAssets;
+		}
+
 		MainStatusText = LOCTEXT("DiscoveringAssetFiles", "Discovering Asset Files");
 		SubStatusText = FText::Format(LOCTEXT("XFilesFoundFmt", "{0} files found"), FText::AsNumber(ProgressUpdateData.NumTotalAssets));
 	}
