Multitude of small changes and fixes to content browser & asset picker.

- Removed unused function

- Added 'FX' asset type category for Niagara purposes that maps to the FX asset category path for default expansion purposes

- Restored functionality for default filter category expansion. This code seemed abandoned. The category expansion could be specified in the asset picker config, but in the end was ignored and the Basic asset path was used instead. This uses a Convert function to go from AssetTypeCategory -> AssetCategoryPath. There is no 1:1 mapping, so it falls back to Basic if necessary. Changing the type in the asset picker config would be cleaner, but isn't straight forward due to FAssetCategoryPath needing to be exposed to all modules that include the content browser module

- Tweaked OnExtendAddFilterMenu delegate. It's now a global delegate, making it unnecessary to handle template shenanigans to access it.
AssetPickerConfig has a new option to provide this OnExtendAddFilterMenu externally. As the 'Common' category still shows up if we expand the FX category (since Niagara System is part of Basic & FX), we can use this to remove unwanted filters
The delegate is called on the dynamic menu section before the widget generation

- Added PinAllFrontendFilters argument. This is useful if you have a small amount of filters available so that you want to display them all by default

Big code cleanup:
- MakeAddFilterMenu is now marked as final to help with API consistency.
- Added new 'CreateAssetFilterBarContext` virtual function. This is now the unified way of configuring your menu.
- The context now contains the FOnFilterAssetType delegate, and the OnExtendAddFilterMenu delegate. This allows subclasses to keep their original functionality, while having a central place to tweak it
- Removed the old content browser context as it's no longer used
- Exposed the menu name. Kept old names consistent. This allows to reuse menus, or create new ones.
- Moved the right click context filter UI of the content browser into the SAssetFilterBar. Now all SAssetFilterBars can do this. Removed the OnGetContextMenu delegate that was only used in one place
- Removed various functions that are no longer needed

diff --git a/Engine/Plugins/VirtualProduction/RemoteControl/Source/RemoteControlUI/Private/UI/Filters/SRCFilterBar.h b/Engine/Plugins/VirtualProduction/RemoteControl/Source/RemoteControlUI/Private/UI/Filters/SRCFilterBar.h
--- a/Engine/Plugins/VirtualProduction/RemoteControl/Source/RemoteControlUI/Private/UI/Filters/SRCFilterBar.h
+++ b/Engine/Plugins/VirtualProduction/RemoteControl/Source/RemoteControlUI/Private/UI/Filters/SRCFilterBar.h
@@ -46,7 +46,6 @@ class SRCFilterBar : public SBasicFilterBar<FilterType>
 public:
 
 	using FOnFilterChanged = typename SBasicFilterBar<FilterType>::FOnFilterChanged;
-	using FOnExtendAddFilterMenu = typename SBasicFilterBar<FilterType>::FOnExtendAddFilterMenu;
 
 	SLATE_BEGIN_ARGS(SRCFilterBar)
 		: _UseDefaultEntityFilters(true)
@@ -56,7 +55,7 @@ public:
 		SLATE_EVENT(SRCFilterBar<FilterType>::FOnFilterChanged, OnFilterChanged)
 
 		/** Delegate to extend the Add Filter dropdown */
-		SLATE_EVENT(SRCFilterBar<FilterType>::FOnExtendAddFilterMenu, OnExtendAddFilterMenu)
+		SLATE_EVENT(FOnExtendAddFilterMenu, OnExtendAddFilterMenu)
 
 		/** Initial List of Custom Filters that will be added to the AddFilter Menu */
 		SLATE_ARGUMENT(TArray<TSharedRef<FFilterBase<FilterType>>>, CustomFilters)
diff --git a/Engine/Source/Developer/AssetTools/Public/AssetTypeCategories.h b/Engine/Source/Developer/AssetTools/Public/AssetTypeCategories.h
--- a/Engine/Source/Developer/AssetTools/Public/AssetTypeCategories.h
+++ b/Engine/Source/Developer/AssetTools/Public/AssetTypeCategories.h
@@ -27,9 +27,10 @@ namespace EAssetTypeCategories
 		Media = 1 << 9,
 		Textures = 1 << 10,
 		World = 1 << 11,
+		FX = 1 << 12,
 
 		// Items below this will be allocated at runtime via RegisterAdvancedAssetCategory
-		FirstUser = 1 << 12,
+		FirstUser = 1 << 13,
 		LastUser = 1 << 31,
 		// Last allowed value is 1 << 31
 	};
diff --git a/Engine/Source/Developer/ToolWidgets/Public/Filters/SBasicFilterBar.h b/Engine/Source/Developer/ToolWidgets/Public/Filters/SBasicFilterBar.h
--- a/Engine/Source/Developer/ToolWidgets/Public/Filters/SBasicFilterBar.h
+++ b/Engine/Source/Developer/ToolWidgets/Public/Filters/SBasicFilterBar.h
@@ -33,6 +33,8 @@
 
 /* Delegate used by SBasicFilterBar to populate the add filter menu */
 DECLARE_DELEGATE_OneParam(FOnPopulateAddFilterMenu, UToolMenu*)
+/** Delegate to extend a menu after it has been generated dynamically. */
+DECLARE_DELEGATE_OneParam( FOnExtendAddFilterMenu, UToolMenu* );
 
 /** ToolMenuContext that is used to create the Add Filter Menu */
 UCLASS()
@@ -42,6 +44,7 @@ class TOOLWIDGETS_API UFilterBarContext : public UObject
 
 public:
 	FOnPopulateAddFilterMenu PopulateFilterMenu;
+	FOnExtendAddFilterMenu OnExtendAddFilterMenu;
 };
 
 /** Describes if the filters are laid out horizontally (ScrollBox) or vertically (WrapBox) */
@@ -76,6 +79,7 @@ struct FFrontendFilterExternalActivationHelper;
 * Use the GetAllActiveFilters() function to get the FilterCollection of Active Filters in this FilterBar, that can be used to filter your items
 * Use MakeAddFilterButton() to make the button that summons the dropdown showing all the filters
 */
+
 template<typename FilterType>
 class SBasicFilterBar : public SCompoundWidget
 {
@@ -83,9 +87,6 @@ public:
 	/** Delegate for when filters have changed */
 	DECLARE_DELEGATE( FOnFilterChanged );
 	
-	/** Delegate for when filters have changed */
-	DECLARE_DELEGATE_OneParam( FOnExtendAddFilterMenu, UToolMenu* );
-	
 	/** Delegate to create a TTextFilter used to compare FilterType with text queries */
     DECLARE_DELEGATE_RetVal( TSharedPtr<ICustomTextFilter<FilterType>>, FCreateTextFilter);
 	
@@ -128,6 +129,9 @@ public:
 
 		/** Whether to use submenus or sections for categories in the filter menu */
 		SLATE_ARGUMENT(bool, UseSectionsForCategories)
+
+		/** Whether we want to automatically pin all frontend filters to make them visible by default. */
+		SLATE_ARGUMENT(bool, bPinAllFrontendFilters)
 	
  	SLATE_END_ARGS()
 
@@ -140,6 +144,7 @@ public:
  		bCanChangeOrientation = InArgs._CanChangeOrientation;
  		FilterPillStyle = InArgs._FilterPillStyle;
 		bUseSectionsForCategories = InArgs._UseSectionsForCategories;
+ 		bPinAllFrontendFilters = InArgs._bPinAllFrontendFilters;
 
  		// We use a widgetswitcher to allow dynamically swapping between the layouts
  		ChildSlot
@@ -202,6 +207,11 @@ public:
 
  		// Auto add if it is an inverse filters
  		SetFrontendFilterActive(InFilter, false);
+
+ 		if(bPinAllFrontendFilters)
+ 		{
+ 			AddFilterToBar(InFilter);
+ 		}
  	}
 
 	EFilterBarLayout GetFilterLayout()
@@ -1576,14 +1586,16 @@ private:
 				if (UFilterBarContext* Context = InMenu->FindContext<UFilterBarContext>())
 				{
 					Context->PopulateFilterMenu.ExecuteIfBound(InMenu);
+					Context->OnExtendAddFilterMenu.ExecuteIfBound(InMenu);
 				}
 			}));
 		}
 
 		UFilterBarContext* FilterBarContext = NewObject<UFilterBarContext>();
 		FilterBarContext->PopulateFilterMenu = FOnPopulateAddFilterMenu::CreateSP(this, &SBasicFilterBar<FilterType>::PopulateAddFilterMenu);
+		FilterBarContext->OnExtendAddFilterMenu = OnExtendAddFilterMenu;
 		FToolMenuContext ToolMenuContext(FilterBarContext);
-
+		
 		return UToolMenus::Get()->GenerateWidget(FilterMenuName, ToolMenuContext);
 	}
 	
@@ -1642,6 +1654,9 @@ protected:
 	/** Whether to use submenus or sections for categories in the filter menu */
 	bool bUseSectionsForCategories;
 
+	/** Whether to pin all front end filters by default. Will add them from the Add Filter menu to the UI, but does not necessarily activate them. */
+	bool bPinAllFrontendFilters;
+
 	/** Determines how each individual pill looks like */
 	EFilterPillStyle FilterPillStyle;
 	
diff --git a/Engine/Source/Editor/ContentBrowser/Private/SAssetPicker.cpp b/Engine/Source/Editor/ContentBrowser/Private/SAssetPicker.cpp
--- a/Engine/Source/Editor/ContentBrowser/Private/SAssetPicker.cpp
+++ b/Engine/Source/Editor/ContentBrowser/Private/SAssetPicker.cpp
@@ -205,7 +205,9 @@ void SAssetPicker::Construct( const FArguments& InArgs )
 			.FrontendFilters(FrontendFilters)
 			.InitialClassFilters(FilterClassList)
 			.FilterBarIdentifier(FName(SaveSettingsName))
-			.ExtraFrontendFilters(InArgs._AssetPickerConfig.ExtraFrontendFilters);
+			.ExtraFrontendFilters(InArgs._AssetPickerConfig.ExtraFrontendFilters)
+			.DefaultMenuExpansionCategory(ConvertAssetTypeCategoryToAssetCategoryPath(DefaultFilterMenuExpansion).Get(EAssetCategoryPaths::Basic))
+			.OnExtendAddFilterMenu(InArgs._AssetPickerConfig.OnExtendAddFilterMenu);
 		
 		FilterComboButtonPtr = StaticCastSharedRef<SComboButton>(SFilterList::MakeAddFilterButton(FilterListPtr.ToSharedRef()));
 		
@@ -529,11 +531,6 @@ void SAssetPicker::SetNewBackendFilter(const FARFilter& NewFilter)
 	OnFilterChanged();
 }
 
-TSharedRef<SWidget> SAssetPicker::MakeAddFilterMenu()
-{
-	return FilterListPtr->ExternalMakeAddFilterMenu(DefaultFilterMenuExpansion);
-}
-
 void SAssetPicker::OnFilterChanged()
 {
 	FARFilter Filter;
@@ -815,4 +812,41 @@ TSharedPtr<SWidget> SAssetPicker::GetItemContextMenu(TArrayView<const FContentBr
 	return nullptr;
 }
 
+TOptional<FAssetCategoryPath> SAssetPicker::ConvertAssetTypeCategoryToAssetCategoryPath(EAssetTypeCategories::Type InDefaultFilterMenuExpansion)
+{
+	// TODO We should completely replace EAssetTypeCategories with FAssetCategoryPath, but FAssetCategoryPath is contained in the AssetDefinitionsModule.
+	// Since the API exposes the DefaultFilterMenuExpansion, we can't easily change this
+	switch (InDefaultFilterMenuExpansion)
+	{
+	case EAssetTypeCategories::Basic:
+		return EAssetCategoryPaths::Basic;
+	case EAssetTypeCategories::Animation:
+		return EAssetCategoryPaths::Animation;
+	case EAssetTypeCategories::Materials:
+		return EAssetCategoryPaths::Material;
+	case EAssetTypeCategories::Sounds:
+		return EAssetCategoryPaths::Audio;
+	case EAssetTypeCategories::Physics:
+		return EAssetCategoryPaths::Physics;
+	case EAssetTypeCategories::UI:
+		return EAssetCategoryPaths::UI;
+	case EAssetTypeCategories::Misc:
+		return EAssetCategoryPaths::Misc;
+	case EAssetTypeCategories::Gameplay:
+		return EAssetCategoryPaths::Gameplay;
+	case EAssetTypeCategories::Blueprint:
+		return EAssetCategoryPaths::Blueprint;
+	case EAssetTypeCategories::Media:
+		return EAssetCategoryPaths::Media;
+	case EAssetTypeCategories::Textures:
+		return EAssetCategoryPaths::Texture;
+	case EAssetTypeCategories::World:
+		break;
+	case EAssetTypeCategories::FX:
+		return EAssetCategoryPaths::FX;
+	}
+
+	return TOptional<FAssetCategoryPath>();
+}
+
 #undef LOCTEXT_NAMESPACE
diff --git a/Engine/Source/Editor/ContentBrowser/Private/SAssetPicker.h b/Engine/Source/Editor/ContentBrowser/Private/SAssetPicker.h
--- a/Engine/Source/Editor/ContentBrowser/Private/SAssetPicker.h
+++ b/Engine/Source/Editor/ContentBrowser/Private/SAssetPicker.h
@@ -85,9 +85,6 @@ private:
 	/** Called from external code to set the filter after the widget was created */
 	void SetNewBackendFilter(const FARFilter& NewFilter);
 
-	/** Called to create the menu for the filter button */
-	TSharedRef<SWidget> MakeAddFilterMenu();
-
 	/** Called when the user changes filters */
 	void OnFilterChanged();
 
@@ -144,6 +141,9 @@ private:
 	/** Handler for the context menu for asset or folder items */
 	TSharedPtr<SWidget> GetItemContextMenu(TArrayView<const FContentBrowserItem> SelectedItems);
 
+	/** Converts an asset type category into an asset category path, if possible. This is done because the asset picker config still uses asset type categories. */
+	TOptional<FAssetCategoryPath> ConvertAssetTypeCategoryToAssetCategoryPath(EAssetTypeCategories::Type DefaultFilterMenuExpansion);
+
 private:
 
 	/** The list of FrontendFilters currently applied to the asset view */
diff --git a/Engine/Source/Editor/ContentBrowser/Private/SContentBrowser.cpp b/Engine/Source/Editor/ContentBrowser/Private/SContentBrowser.cpp
--- a/Engine/Source/Editor/ContentBrowser/Private/SContentBrowser.cpp
+++ b/Engine/Source/Editor/ContentBrowser/Private/SContentBrowser.cpp
@@ -961,7 +961,6 @@ TSharedRef<SWidget> SContentBrowser::CreateAssetView(const FContentBrowserConfig
 	// Create the Filter Bar Widget
 	FilterListPtr = SNew(SFilterList)
 					.OnFilterChanged(this, &SContentBrowser::OnFilterChanged)
-					.OnGetContextMenu(this, &SContentBrowser::GetFilterContextMenu)
 					.Visibility((Config != nullptr ? Config->bCanShowFilters : true) ? EVisibility::Visible : EVisibility::Collapsed)
 					.FrontendFilters(FrontendFilters)
 					.FilterBarIdentifier(InstanceName)
@@ -2952,16 +2951,6 @@ FText SContentBrowser::GetAddNewToolTipText() const
 	return LOCTEXT( "AddNewToolTip_NoPath", "No path is selected as an add target." );
 }
 
-TSharedRef<SWidget> SContentBrowser::MakeAddFilterMenu()
-{
-	return FilterListPtr->ExternalMakeAddFilterMenu();
-}
-
-TSharedPtr<SWidget> SContentBrowser::GetFilterContextMenu()
-{
-	return FilterListPtr->ExternalMakeAddFilterMenu();
-}
-
 void SContentBrowser::RegisterPathViewFiltersMenu()
 {
 	static const FName PathViewFiltersMenuName = TEXT("ContentBrowser.AssetViewOptions.PathViewFilters");
diff --git a/Engine/Source/Editor/ContentBrowser/Private/SContentBrowser.h b/Engine/Source/Editor/ContentBrowser/Private/SContentBrowser.h
--- a/Engine/Source/Editor/ContentBrowser/Private/SContentBrowser.h
+++ b/Engine/Source/Editor/ContentBrowser/Private/SContentBrowser.h
@@ -343,12 +343,6 @@ private:
 	/** Gets the tool tip for the "Add New" button */
 	FText GetAddNewToolTipText() const;
 
-	/** Makes the filters menu */
-	TSharedRef<SWidget> MakeAddFilterMenu();
-	
-	/** Builds the context menu for the filter list area. */
-	TSharedPtr<SWidget> GetFilterContextMenu();
-
 	/** Saves dirty content. */
 	FReply OnSaveClicked();
 
diff --git a/Engine/Source/Editor/ContentBrowser/Private/SFilterList.cpp b/Engine/Source/Editor/ContentBrowser/Private/SFilterList.cpp
--- a/Engine/Source/Editor/ContentBrowser/Private/SFilterList.cpp
+++ b/Engine/Source/Editor/ContentBrowser/Private/SFilterList.cpp
@@ -62,11 +62,10 @@ SFilterList::FCustomTextFilterEvent SFilterList::CustomTextFilterEvent;
 void SFilterList::Construct( const FArguments& InArgs )
 {
 	bUseSharedSettings = InArgs._UseSharedSettings;
-	OnGetContextMenu = InArgs._OnGetContextMenu;
 	OnFilterBarLayoutChanging = InArgs._OnFilterBarLayoutChanging;
 	this->OnFilterChanged = InArgs._OnFilterChanged;
 	this->ActiveFilters = InArgs._FrontendFilters;
-	InitialClassFilters = InArgs._InitialClassFilters;
+	InitialClassFilters = InArgs._InitialClassFilters; 
 
 	TSharedPtr<FFrontendFilterCategory> DefaultCategory = MakeShareable( new FFrontendFilterCategory(LOCTEXT("FrontendFiltersCategory", "Other Filters"), LOCTEXT("FrontendFiltersCategoryTooltip", "Filter assets by all filters in this category.")) );
 	
@@ -139,6 +138,9 @@ void SFilterList::Construct( const FArguments& InArgs )
 	Args._FilterBarIdentifier = InArgs._FilterBarIdentifier;
 	Args._FilterBarLayout = InArgs._FilterBarLayout;
 	Args._CanChangeOrientation = InArgs._CanChangeOrientation;
+	Args._OnExtendAddFilterMenu = InArgs._OnExtendAddFilterMenu;
+	Args._FilterMenuName = FName("ContentBrowser.FilterMenu");
+	Args._DefaultMenuExpansionCategory = InArgs._DefaultMenuExpansionCategory;
 
 	SAssetFilterBar<FAssetFilterType>::Construct(Args);
 
@@ -155,31 +157,6 @@ void SFilterList::Construct( const FArguments& InArgs )
 	
 }
 
-FReply SFilterList::OnMouseButtonUp( const FGeometry& MyGeometry, const FPointerEvent& MouseEvent )
-{
-	if ( MouseEvent.GetEffectingButton() == EKeys::RightMouseButton )
-	{
-		if ( OnGetContextMenu.IsBound() )
-		{
-			FReply Reply = FReply::Handled().ReleaseMouseCapture();
-
-			// Get the context menu content. If NULL, don't open a menu.
-			TSharedPtr<SWidget> MenuContent = OnGetContextMenu.Execute();
-
-			if ( MenuContent.IsValid() )
-			{
-				FVector2D SummonLocation = MouseEvent.GetScreenSpacePosition();
-				FWidgetPath WidgetPath = MouseEvent.GetEventPath() != nullptr ? *MouseEvent.GetEventPath() : FWidgetPath();
-				FSlateApplication::Get().PushMenu(AsShared(), WidgetPath, MenuContent.ToSharedRef(), SummonLocation, FPopupTransitionEffect(FPopupTransitionEffect::ContextMenu));
-			}
-
-			return Reply;
-		}
-	}
-
-	return FReply::Unhandled();
-}
-
 const TArray<UClass*>& SFilterList::GetInitialClassFilters()
 {
 	return InitialClassFilters;
@@ -197,9 +174,9 @@ TSharedPtr<FFrontendFilter> SFilterList::GetFrontendFilter(const FString& InName
 	return TSharedPtr<FFrontendFilter>();
 }
 
-TSharedRef<SWidget> SFilterList::ExternalMakeAddFilterMenu(EAssetTypeCategories::Type MenuExpansion)
+TSharedRef<SWidget> SFilterList::ExternalMakeAddFilterMenu()
 {
-	return MakeAddFilterMenu(MenuExpansion);
+	return SAssetFilterBar<FAssetFilterType>::MakeAddFilterMenu();
 }
 
 void SFilterList::DisableFiltersThatHideItems(TArrayView<const FContentBrowserItem> ItemList)
@@ -328,78 +305,16 @@ bool IsFilteredByPicker(const TArray<UClass*>& FilterClassList, UClass* TestClas
 	return true;
 }
 
-void SFilterList::PopulateAddFilterMenu_Internal(UToolMenu* Menu)
+UAssetFilterBarContext* SFilterList::CreateAssetFilterBarContext()
 {
-	EAssetTypeCategories::Type MenuExpansion = EAssetTypeCategories::Basic;
-	if (UContentBrowserFilterListContext* Context = Menu->FindContext<UContentBrowserFilterListContext>())
-	{
-		MenuExpansion = Context->MenuExpansion;
-	}
+	UAssetFilterBarContext* AssetFilterBarContext = SAssetFilterBar<const FContentBrowserItem&>::CreateAssetFilterBarContext();
 
-	this->PopulateAddFilterMenu(Menu, AssetFilterCategories.FindChecked(EAssetCategoryPaths::Basic.GetCategory()), FOnFilterAssetType::CreateLambda([this](UClass *TestClass)
+	AssetFilterBarContext->OnFilterAssetType = FOnFilterAssetType::CreateLambda([this](UClass *TestClass)
 	{
 		return !IsFilteredByPicker(this->InitialClassFilters, TestClass);
-	}));
+	});
 	
-	Menu->AddSection("ContentBrowserFilterMiscAsset", LOCTEXT("MiscAssetsMenuHeading", "Misc Options") );
-}
-
-TSharedRef<SWidget> SFilterList::MakeAddFilterMenu()
-{
-	const FName FilterMenuName = "ContentBrowser.FilterMenu";
-	if (!UToolMenus::Get()->IsMenuRegistered(FilterMenuName))
-	{
-		UToolMenu* Menu = UToolMenus::Get()->RegisterMenu(FilterMenuName);
-		Menu->bShouldCloseWindowAfterMenuSelection = true;
-		Menu->bCloseSelfOnly = true;
-
-		Menu->AddDynamicSection(NAME_None, FNewToolMenuDelegate::CreateLambda([](UToolMenu* InMenu)
-		{
-			if (UContentBrowserFilterListContext* Context = InMenu->FindContext<UContentBrowserFilterListContext>())
-			{
-				if (TSharedPtr<SFilterList> FilterList = Context->FilterList.Pin())
-				{
-					FilterList->PopulateAddFilterMenu_Internal(InMenu);
-				}
-			}
-		}));
-	}
-
-	UContentBrowserFilterListContext* ContentBrowserFilterListContext = NewObject<UContentBrowserFilterListContext>();
-	ContentBrowserFilterListContext->FilterList = SharedThis(this);
-	ContentBrowserFilterListContext->MenuExpansion = EAssetTypeCategories::Basic;
-	FToolMenuContext ToolMenuContext(ContentBrowserFilterListContext);
-
-	return UToolMenus::Get()->GenerateWidget(FilterMenuName, ToolMenuContext);
-}
-
-TSharedRef<SWidget> SFilterList::MakeAddFilterMenu(EAssetTypeCategories::Type MenuExpansion)
-{
-	const FName FilterMenuName = "ContentBrowser.FilterMenu";
-	if (!UToolMenus::Get()->IsMenuRegistered(FilterMenuName))
-	{
-		UToolMenu* Menu = UToolMenus::Get()->RegisterMenu(FilterMenuName);
-		Menu->bShouldCloseWindowAfterMenuSelection = true;
-		Menu->bCloseSelfOnly = true;
-
-		Menu->AddDynamicSection(NAME_None, FNewToolMenuDelegate::CreateLambda([](UToolMenu* InMenu)
-		{
-			if (UContentBrowserFilterListContext* Context = InMenu->FindContext<UContentBrowserFilterListContext>())
-			{
-				if (TSharedPtr<SFilterList> FilterList = Context->FilterList.Pin())
-				{
-					FilterList->PopulateAddFilterMenu_Internal(InMenu);
-				}
-			}
-		}));
-	}
-
-	UContentBrowserFilterListContext* ContentBrowserFilterListContext = NewObject<UContentBrowserFilterListContext>();
-	ContentBrowserFilterListContext->FilterList = SharedThis(this);
-	ContentBrowserFilterListContext->MenuExpansion = MenuExpansion;
-	FToolMenuContext ToolMenuContext(ContentBrowserFilterListContext);
-
-	return UToolMenus::Get()->GenerateWidget(FilterMenuName, ToolMenuContext);
+	return AssetFilterBarContext;
 }
 
 void SFilterList::CreateCustomFilterDialog(const FText& InText)
diff --git a/Engine/Source/Editor/ContentBrowser/Private/SFilterList.h b/Engine/Source/Editor/ContentBrowser/Private/SFilterList.h
--- a/Engine/Source/Editor/ContentBrowser/Private/SFilterList.h
+++ b/Engine/Source/Editor/ContentBrowser/Private/SFilterList.h
@@ -41,8 +41,6 @@ enum class ECheckBoxState : uint8;
 class SFilterList : public SAssetFilterBar<FAssetFilterType>
 {
 public:
-
-	DECLARE_DELEGATE_RetVal( TSharedPtr<SWidget>, FOnGetContextMenu );
 	DECLARE_DELEGATE_OneParam(FOnFilterBarLayoutChanging, EFilterBarLayout /* NewLayout */)
 	/**
 	 * An event delegate that is executed when a custom text filter has been created/modified/deleted in any FilterList
@@ -56,16 +54,16 @@ public:
 	: _UseSharedSettings(false)
 	, _FilterBarLayout(EFilterBarLayout::Horizontal)
 	, _CanChangeOrientation(false)
+	, _DefaultMenuExpansionCategory(EAssetCategoryPaths::Basic)
 	{
 		
 	}
-
-		/** Called when an asset is right clicked */
-		SLATE_EVENT( FOnGetContextMenu, OnGetContextMenu )
-
 		/** Delegate for when filters have changed */
 		SLATE_EVENT( FOnFilterChanged, OnFilterChanged )
 
+		/** Delegate that lets the user modify the menu after the fact */
+		SLATE_EVENT(FOnExtendAddFilterMenu, OnExtendAddFilterMenu)
+	
 		/** The filter collection used to further filter down assets returned from the backend */
 		SLATE_ARGUMENT( TSharedPtr<FAssetFilterCollectionType>, FrontendFilters)
 
@@ -97,6 +95,9 @@ public:
 		/** If true, allow dynamically changing the orientation and saving in the config */
 		SLATE_ARGUMENT(bool, CanChangeOrientation)
 
+		/** Expands the specified asset category, if specified. If not, it will expand Basic/Common instead. */
+		SLATE_ARGUMENT(TOptional<FAssetCategoryPath>, DefaultMenuExpansionCategory)
+
 	SLATE_END_ARGS()
 
 	/** Constructs this widget with InArgs */
@@ -113,15 +114,13 @@ public:
 
 	/** Retrieve a specific frontend filter */
 	TSharedPtr<FFrontendFilter> GetFrontendFilter(const FString& InName) const;
-
+	
 	/** Handler for when the floating add filter button was clicked */
-	TSharedRef<SWidget> ExternalMakeAddFilterMenu(EAssetTypeCategories::Type MenuExpansion = EAssetTypeCategories::Basic);
+	TSharedRef<SWidget> ExternalMakeAddFilterMenu();
 
 	/** Disables any active filters that would hide the supplied items */
 	void DisableFiltersThatHideItems(TArrayView<const FContentBrowserItem> ItemList);
 
-	virtual FReply OnMouseButtonUp( const FGeometry& MyGeometry, const FPointerEvent& MouseEvent ) override;
-
 	/** Returns the class filters specified at construction using argument 'InitialClassFilters'. */
 	const TArray<UClass*>& GetInitialClassFilters();
 
@@ -150,7 +149,7 @@ public:
 protected:
 	
 	/** Handler for when the add filter button was clicked */
-	TSharedRef<SWidget> MakeAddFilterMenu() override;
+	//virtual TSharedRef<SWidget> MakeAddFilterMenu() override;
 	
 	/** Handler for when a custom text filter is created */
 	virtual void OnCreateCustomTextFilter(const FCustomTextFilterData& InFilterData, bool bApplyFilter) override;
@@ -161,11 +160,11 @@ protected:
 
 private:
 
-	/** Exists for backwards compatibility with ExternalMakeAddFilterMenu */
-	TSharedRef<SWidget> MakeAddFilterMenu(EAssetTypeCategories::Type MenuExpansion = EAssetTypeCategories::Basic);
-
-	void PopulateAddFilterMenu_Internal(UToolMenu* Menu);
+	// /** Exists for backwards compatibility with ExternalMakeAddFilterMenu */
+	// TSharedRef<SWidget> MakeAddFilterMenu(FAssetCategoryPath MenuExpansion = EAssetCategoryPaths::Basic);
 
+	virtual UAssetFilterBarContext* CreateAssetFilterBarContext() override;
+	
 	/** Handler for when another SFilterList using shared settings creates a custom text filter */
 	void OnExternalCustomTextFilterCreated(TSharedPtr<SWidget> BroadcastingFilterList);
 	
@@ -182,9 +181,6 @@ private:
 	/** List of classes that our filters must match */
 	TArray<UClass*> InitialClassFilters;
 
-	/** Delegate for getting the context menu. */
-	FOnGetContextMenu OnGetContextMenu;
-
 	/** Delegate for when filters have changed */
 	FOnFilterChanged OnFilterChanged;
 
diff --git a/Engine/Source/Editor/ContentBrowser/Public/ContentBrowserMenuContexts.h b/Engine/Source/Editor/ContentBrowser/Public/ContentBrowserMenuContexts.h
--- a/Engine/Source/Editor/ContentBrowser/Public/ContentBrowserMenuContexts.h
+++ b/Engine/Source/Editor/ContentBrowser/Public/ContentBrowserMenuContexts.h
@@ -307,18 +307,6 @@ public:
 	const TArray<FString>& GetSelectedPackagePaths() const { return SelectedPackagePaths; }
 };
 
-UCLASS()
-class CONTENTBROWSER_API UContentBrowserFilterListContext : public UObject
-{
-	GENERATED_BODY()
-
-public:
-
-	TWeakPtr<SFilterList> FilterList;
-
-	EAssetTypeCategories::Type MenuExpansion;
-};
-
 UCLASS()
 class CONTENTBROWSER_API UContentBrowserAddNewContextMenuContext : public UObject
 {
diff --git a/Engine/Source/Editor/ContentBrowser/Public/IContentBrowserSingleton.h b/Engine/Source/Editor/ContentBrowser/Public/IContentBrowserSingleton.h
--- a/Engine/Source/Editor/ContentBrowser/Public/IContentBrowserSingleton.h
+++ b/Engine/Source/Editor/ContentBrowser/Public/IContentBrowserSingleton.h
@@ -22,6 +22,7 @@ class FViewport;
 class IPlugin;
 class SWidget;
 class UFactory;
+class UToolMenu;
 
 typedef const FContentBrowserItem& FAssetFilterType;
 typedef TFilterCollection<FAssetFilterType> FAssetFilterCollectionType;
@@ -296,6 +297,9 @@ struct FAssetPickerConfig
 	/** If true, sort by path in column view. Only works if initial view type is Column */
 	bool bSortByPathInColumnView;
 
+	/** Can be used to freely change the Add Filter menu. This is executed on the dynamically instanced menu, not the registered menu. */
+	TDelegate<void(UToolMenu*)> OnExtendAddFilterMenu;
+	
 	/** Override the default filter context menu layout */
 	EAssetTypeCategories::Type DefaultFilterMenuExpansion;
 
diff --git a/Engine/Source/Editor/EditorWidgets/Public/Filters/SAssetFilterBar.h b/Engine/Source/Editor/EditorWidgets/Public/Filters/SAssetFilterBar.h
--- a/Engine/Source/Editor/EditorWidgets/Public/Filters/SAssetFilterBar.h
+++ b/Engine/Source/Editor/EditorWidgets/Public/Filters/SAssetFilterBar.h
@@ -77,6 +77,8 @@ class EDITORWIDGETS_API UAssetFilterBarContext : public UObject
 public:
 	TSharedPtr<FFilterCategory> MenuExpansion;
 	FOnPopulateAddAssetFilterMenu PopulateFilterMenu;
+	FOnFilterAssetType OnFilterAssetType;
+	FOnExtendAddFilterMenu OnExtendAddFilterMenu;
 };
 
 /** A non-templated base class for SAssetFilterBar, where functionality that does not depend on the type of item
@@ -144,14 +146,15 @@ class SAssetFilterBar : public SBasicFilterBar<FilterType>, public FFilterBarBas
 public:
 	
 	using FOnFilterChanged = typename SBasicFilterBar<FilterType>::FOnFilterChanged;
-	using FOnExtendAddFilterMenu = typename SBasicFilterBar<FilterType>::FOnExtendAddFilterMenu;
 	using FCreateTextFilter = typename SBasicFilterBar<FilterType>::FCreateTextFilter;
 	
 	SLATE_BEGIN_ARGS( SAssetFilterBar<FilterType> )
-		: _UseDefaultAssetFilters(true)
+		: _FilterMenuName(FName("FilterBar.FilterMenu"))
+		, _UseDefaultAssetFilters(true)
 		, _FilterBarLayout(EFilterBarLayout::Horizontal)
 		, _CanChangeOrientation(false)
 		, _FilterPillStyle(EFilterPillStyle::Default)
+		, _DefaultMenuExpansionCategory(EAssetCategoryPaths::Basic)
 		{
 		
 		}
@@ -179,6 +182,9 @@ public:
 		 */
 		SLATE_ARGUMENT(TSharedPtr<SFilterSearchBox>, FilterSearchBox)
 
+		/** The filter menu name. Will register a menu using this name, if not already registered. */
+		SLATE_ARGUMENT(FName, FilterMenuName)
+	
 		/** A unique identifier for this filter bar needed to enable saving settings in a config file */
 		SLATE_ARGUMENT(FName, FilterBarIdentifier)
 	
@@ -193,6 +199,9 @@ public:
 
 		/** Determines how each individual filter pill looks like */
 		SLATE_ARGUMENT(EFilterPillStyle, FilterPillStyle)
+
+		/** Expands the specified asset category, if specified. If not, it will expand Basic/Common instead. */
+		SLATE_ARGUMENT(TOptional<FAssetCategoryPath>, DefaultMenuExpansionCategory)
 	
 	SLATE_END_ARGS()
 
@@ -202,7 +211,9 @@ public:
 		bUseDefaultAssetFilters = InArgs._UseDefaultAssetFilters;
 		CustomClassFilters = InArgs._CustomClassFilters;
 		UserCustomClassFilters = InArgs._CustomClassFilters;
+		FilterMenuName = InArgs._FilterMenuName;
 		FilterBarIdentifier = InArgs._FilterBarIdentifier;
+		DefaultMenuExpansionCategory = InArgs._DefaultMenuExpansionCategory;
 		
 		typename SBasicFilterBar<FilterType>::FArguments Args;
 		Args._OnFilterChanged = InArgs._OnFilterChanged;
@@ -285,6 +296,27 @@ public:
 
 protected:
 
+	virtual FReply OnMouseButtonUp(const FGeometry& MyGeometry, const FPointerEvent& MouseEvent) override
+	{
+		if(MouseEvent.GetEffectingButton() == EKeys::RightMouseButton )
+		{
+			FReply Reply = FReply::Handled().ReleaseMouseCapture();
+			
+			TSharedPtr<SWidget> MenuContent = MakeAddFilterMenu();
+
+			if(MenuContent.IsValid())
+			{
+				FVector2D SummonLocation = MouseEvent.GetScreenSpacePosition();
+				FWidgetPath WidgetPath = MouseEvent.GetEventPath() != nullptr ? *MouseEvent.GetEventPath() : FWidgetPath();
+				FSlateApplication::Get().PushMenu(this->AsShared(), WidgetPath, MenuContent.ToSharedRef(), SummonLocation, FPopupTransitionEffect(FPopupTransitionEffect::ContextMenu));
+			}
+
+			return Reply;
+		}
+
+		return FReply::Unhandled();
+	}
+	
 	/** Save all the custom class filters (created by the user) into the specified config */
 	void SaveCustomTextFilters(FFilterBarSettings* FilterBarConfig)
 	{
@@ -815,9 +847,6 @@ protected:
 			// Add the Basic category
 			TSharedPtr<FFilterCategory> BasicCategory = MakeShareable(new FFilterCategory(LOCTEXT("BasicFilter", "Common"), LOCTEXT("BasicFilterTooltip", "Filter by Common assets.")));
 			AssetFilterCategories.Add(EAssetCategoryPaths::Basic.GetCategory(), BasicCategory);
-
-			// We expand the basic category by default if we are creating it
-			CategoryToExpand = BasicCategory;
 		}
 
 		{
@@ -917,10 +946,26 @@ protected:
 		CreateFiltersMenuCategory(InMenu->AddSection("Section"), CustomClassFilterDatas);
 	}
 
+	/** Override this if you want to change aspects on the add filter menu.
+	 *  You can opt to only override the OnFilterAssetType member to exclude assets, or choose to override the entire populate callback
+	 */
+	virtual UAssetFilterBarContext* CreateAssetFilterBarContext()
+	{
+		UAssetFilterBarContext* FilterBarContext = NewObject<UAssetFilterBarContext>();
+		FilterBarContext->PopulateFilterMenu = FOnPopulateAddAssetFilterMenu::CreateSP(this, &SAssetFilterBar<FilterType>::PopulateAddFilterMenu);
+		FilterBarContext->OnExtendAddFilterMenu = this->OnExtendAddFilterMenu;
+		
+		if(DefaultMenuExpansionCategory.IsSet() && AssetFilterCategories.Contains(DefaultMenuExpansionCategory->GetCategory()))
+		{
+			FilterBarContext->MenuExpansion = AssetFilterCategories[DefaultMenuExpansionCategory->GetCategory()];
+		};
+
+		return FilterBarContext;
+	}
+	
 	/** Handler for when the add filter button was clicked */
-	virtual TSharedRef<SWidget> MakeAddFilterMenu() override
+	virtual TSharedRef<SWidget> MakeAddFilterMenu() override final
 	{
-		const FName FilterMenuName = "FilterBar.FilterMenu";
 		if (!UToolMenus::Get()->IsMenuRegistered(FilterMenuName))
 		{
 			UToolMenu* Menu = UToolMenus::Get()->RegisterMenu(FilterMenuName);
@@ -931,28 +976,18 @@ protected:
 			{
 				if (UAssetFilterBarContext* Context = InMenu->FindContext<UAssetFilterBarContext>())
 				{
-					Context->PopulateFilterMenu.ExecuteIfBound(InMenu, Context->MenuExpansion, FOnFilterAssetType());
+					Context->PopulateFilterMenu.ExecuteIfBound(InMenu, Context->MenuExpansion, Context->OnFilterAssetType);
+					Context->OnExtendAddFilterMenu.ExecuteIfBound(InMenu);
 				}
 			}));
 		}
-
-		UAssetFilterBarContext* FilterBarContext = NewObject<UAssetFilterBarContext>();
-		FilterBarContext->PopulateFilterMenu = FOnPopulateAddAssetFilterMenu::CreateSP(this, &SAssetFilterBar<FilterType>::PopulateAddFilterMenu);
-
-		/** Auto expand the Basic Category if it is present */
-
-		if(CategoryToExpand)
-		{
-			FilterBarContext->MenuExpansion = CategoryToExpand;
-		}
-	
-		FToolMenuContext ToolMenuContext(FilterBarContext);
-
+		
+		FToolMenuContext ToolMenuContext(CreateAssetFilterBarContext());
 		return UToolMenus::Get()->GenerateWidget(FilterMenuName, ToolMenuContext);
 	} 
 	
 	/** Handler to Populate the Add Filter Menu. Use OnFilterAssetType in Subclasses to add classes to the exclusion list */
-	virtual void PopulateAddFilterMenu(UToolMenu* Menu, TSharedPtr<FFilterCategory> MenuExpansion, FOnFilterAssetType OnFilterAssetType)
+	void PopulateAddFilterMenu(UToolMenu* Menu, TSharedPtr<FFilterCategory> MenuExpansion, FOnFilterAssetType OnFilterAssetType)
 	{
 		FAssetToolsModule& AssetToolsModule = FModuleManager::LoadModuleChecked<FAssetToolsModule>(TEXT("AssetTools"));
 
@@ -1037,6 +1072,7 @@ protected:
 		}
 		
 		// If we want to expand a category
+		// if(MenuExpansion)
 		if(MenuExpansion)
 		{
 			// First add the expanded category, this appears as standard entries in the list (Note: intentionally not using FindChecked here as removing it from the map later would cause the ref to be garbage)
@@ -1080,7 +1116,7 @@ protected:
 		for (const TPair<TSharedPtr<FFilterCategory>, FCategoryMenu>& CategoryMenuPair : CategoryToMenuMap)
 		{
 			Section.AddSubMenu(
-				NAME_None,
+				FName(FText::AsCultureInvariant(CategoryMenuPair.Key->Title).ToString()),
 				CategoryMenuPair.Key->Title,
 				CategoryMenuPair.Key->Tooltip,
 				FNewToolMenuDelegate::CreateSP(this, &SAssetFilterBar<FilterType>::CreateFiltersMenuCategory, CategoryMenuPair.Value.Classes),
@@ -1214,10 +1250,9 @@ protected:
 	/** Copy of Custom Class Filters that were provided by the user, and not autopopulated from AssetTypeActions */
 	TArray<TSharedRef<FCustomClassFilterData>> UserCustomClassFilters;
 
-	TMap<FName, TSharedPtr<FFilterCategory>> AssetFilterCategories;
+	FName FilterMenuName;
 
-	/** The category to expand in the filter menu */
-	TSharedPtr<FFilterCategory> CategoryToExpand;
+	TMap<FName, TSharedPtr<FFilterCategory>> AssetFilterCategories;
 	
 	/** Whether the filter bar provides the default Asset Filters */
 	bool bUseDefaultAssetFilters = false;
@@ -1227,8 +1262,11 @@ protected:
 
 	/** A map of any type filters we encounter while calling LoadSettings(), but don't have in this instance of the filter
 	 *  bar. These are stored so we can save these when calling SaveSettings() so that they are not lost */
-	TMap<FString, bool> UnknownTypeFilters; 
+	TMap<FString, bool> UnknownTypeFilters;
 
+private:
+	/** The filter menu category to expand. */
+	TOptional<FAssetCategoryPath> DefaultMenuExpansionCategory;
 };
 
 #undef LOCTEXT_NAMESPACE
\ No newline at end of file
diff --git a/Engine/Source/Editor/EditorWidgets/Public/Filters/SFilterBar.h b/Engine/Source/Editor/EditorWidgets/Public/Filters/SFilterBar.h
--- a/Engine/Source/Editor/EditorWidgets/Public/Filters/SFilterBar.h
+++ b/Engine/Source/Editor/EditorWidgets/Public/Filters/SFilterBar.h
@@ -40,7 +40,6 @@ public:
 	using FOnFilterChanged = typename SBasicFilterBar<FilterType>::FOnFilterChanged;
 	using FConvertItemToAssetData = typename FAssetFilter<FilterType>::FConvertItemToAssetData;
 	using FCompareItemWithClassNames = typename FAssetFilter<FilterType>::FCompareItemWithClassNames;
-	using FOnExtendAddFilterMenu = typename SBasicFilterBar<FilterType>::FOnExtendAddFilterMenu;
 	using FCreateTextFilter = typename SBasicFilterBar<FilterType>::FCreateTextFilter;
 	using SFilter = typename SBasicFilterBar<FilterType>::SFilter;
 	
@@ -50,6 +49,7 @@ public:
 	    , _FilterBarLayout(EFilterBarLayout::Horizontal)
 		, _CanChangeOrientation(false)
 		, _FilterPillStyle(EFilterPillStyle::Default)
+		, _DefaultMenuExpansionCategory(EAssetCategoryPaths::Basic)
 		{
 		
 		}
@@ -97,6 +97,9 @@ public:
 		
 		/** Determines how each individual filter pill looks like */
 		SLATE_ARGUMENT(EFilterPillStyle, FilterPillStyle)
+
+		/** The filter menu category to expand. */
+		SLATE_ARGUMENT(TOptional<FAssetCategoryPath>, DefaultMenuExpansionCategory);
 	
 	SLATE_END_ARGS()
 
@@ -115,6 +118,7 @@ public:
 		Args._FilterBarLayout = InArgs._FilterBarLayout;
         Args._CanChangeOrientation = InArgs._CanChangeOrientation;
 		Args._FilterPillStyle = InArgs._FilterPillStyle;
+		Args._DefaultMenuExpansionCategory = InArgs._DefaultMenuExpansionCategory;
 		
 		SAssetFilterBar<FilterType>::Construct(Args);
 		
diff --git a/Engine/Source/Editor/SceneOutliner/Private/SceneOutlinerFilterBar.cpp b/Engine/Source/Editor/SceneOutliner/Private/SceneOutlinerFilterBar.cpp
--- a/Engine/Source/Editor/SceneOutliner/Private/SceneOutlinerFilterBar.cpp
+++ b/Engine/Source/Editor/SceneOutliner/Private/SceneOutlinerFilterBar.cpp
@@ -143,6 +143,14 @@ void SSceneOutlinerFilterBar::LoadSettings()
 	this->OnFilterChanged.ExecuteIfBound();
 }
 
+UAssetFilterBarContext* SSceneOutlinerFilterBar::CreateAssetFilterBarContext()
+{
+	UAssetFilterBarContext* AssetFilterBarContext = SFilterBar<const ISceneOutlinerTreeItem&>::CreateAssetFilterBarContext();
+	AssetFilterBarContext->MenuExpansion = CategoryToExpand;
+
+	return AssetFilterBarContext;
+}
+
 void SSceneOutlinerFilterBar::OnCreateCustomTextFilter(const FCustomTextFilterData& InFilterData, bool bApplyFilter)
 {
 	SAssetFilterBar<SceneOutliner::FilterBarType>::OnCreateCustomTextFilter(InFilterData, bApplyFilter);
diff --git a/Engine/Source/Editor/SceneOutliner/Private/SceneOutlinerFilterBar.h b/Engine/Source/Editor/SceneOutliner/Private/SceneOutlinerFilterBar.h
--- a/Engine/Source/Editor/SceneOutliner/Private/SceneOutlinerFilterBar.h
+++ b/Engine/Source/Editor/SceneOutliner/Private/SceneOutlinerFilterBar.h
@@ -74,6 +74,7 @@ public:
 	virtual void LoadSettings() override;
 
 protected:
+	virtual UAssetFilterBarContext* CreateAssetFilterBarContext() override;
 
 	/** Handler for when a custom text filter is created */
 	virtual void OnCreateCustomTextFilter(const FCustomTextFilterData& InFilterData, bool bApplyFilter) override;
@@ -103,5 +104,6 @@ private:
 
 	/** Whether this Filter bar wants to load/save from the settings common to all instances */
 	bool bUseSharedSettings;
-	
+
+	TSharedPtr<FFilterCategory> CategoryToExpand;	
 };
\ No newline at end of file
