- Exposed 'use sections for categories' bool to the asset picker config
- Added identifier to the expanded category (previously None)

diff --git a/Engine/Source/Developer/ToolWidgets/Public/Filters/SBasicFilterBar.h b/Engine/Source/Developer/ToolWidgets/Public/Filters/SBasicFilterBar.h
--- a/Engine/Source/Developer/ToolWidgets/Public/Filters/SBasicFilterBar.h
+++ b/Engine/Source/Developer/ToolWidgets/Public/Filters/SBasicFilterBar.h
@@ -1531,8 +1531,6 @@ protected:
 					FNewToolMenuDelegate::CreateSP(this, &SBasicFilterBar<FilterType>::CreateTextFiltersMenu)
 					);
 		}
-		
-		OnExtendAddFilterMenu.ExecuteIfBound(Menu);
 	}
 
 	/** Helper function to add all custom filters to the Add Filter Menu */
diff --git a/Engine/Source/Editor/ContentBrowser/Private/SAssetPicker.cpp b/Engine/Source/Editor/ContentBrowser/Private/SAssetPicker.cpp
--- a/Engine/Source/Editor/ContentBrowser/Private/SAssetPicker.cpp
+++ b/Engine/Source/Editor/ContentBrowser/Private/SAssetPicker.cpp
@@ -207,7 +207,8 @@ void SAssetPicker::Construct( const FArguments& InArgs )
 			.FilterBarIdentifier(FName(SaveSettingsName))
 			.ExtraFrontendFilters(InArgs._AssetPickerConfig.ExtraFrontendFilters)
 			.DefaultMenuExpansionCategory(ConvertAssetTypeCategoryToAssetCategoryPath(DefaultFilterMenuExpansion).Get(EAssetCategoryPaths::Basic))
-			.OnExtendAddFilterMenu(InArgs._AssetPickerConfig.OnExtendAddFilterMenu);
+			.OnExtendAddFilterMenu(InArgs._AssetPickerConfig.OnExtendAddFilterMenu)
+			.bUseSectionsForCustomCategories(InArgs._AssetPickerConfig.bUseSectionsForCustomFilterCategories);
 		
 		FilterComboButtonPtr = StaticCastSharedRef<SComboButton>(SFilterList::MakeAddFilterButton(FilterListPtr.ToSharedRef()));
 		
diff --git a/Engine/Source/Editor/ContentBrowser/Private/SFilterList.cpp b/Engine/Source/Editor/ContentBrowser/Private/SFilterList.cpp
--- a/Engine/Source/Editor/ContentBrowser/Private/SFilterList.cpp
+++ b/Engine/Source/Editor/ContentBrowser/Private/SFilterList.cpp
@@ -141,6 +141,7 @@ void SFilterList::Construct( const FArguments& InArgs )
 	Args._OnExtendAddFilterMenu = InArgs._OnExtendAddFilterMenu;
 	Args._FilterMenuName = FName("ContentBrowser.FilterMenu");
 	Args._DefaultMenuExpansionCategory = InArgs._DefaultMenuExpansionCategory;
+	Args._bUseSectionsForCustomCategories = InArgs._bUseSectionsForCustomCategories;
 
 	SAssetFilterBar<FAssetFilterType>::Construct(Args);
 
diff --git a/Engine/Source/Editor/ContentBrowser/Private/SFilterList.h b/Engine/Source/Editor/ContentBrowser/Private/SFilterList.h
--- a/Engine/Source/Editor/ContentBrowser/Private/SFilterList.h
+++ b/Engine/Source/Editor/ContentBrowser/Private/SFilterList.h
@@ -55,6 +55,7 @@ public:
 	, _FilterBarLayout(EFilterBarLayout::Horizontal)
 	, _CanChangeOrientation(false)
 	, _DefaultMenuExpansionCategory(EAssetCategoryPaths::Basic)
+	, _bUseSectionsForCustomCategories(false)
 	{
 		
 	}
@@ -98,6 +99,9 @@ public:
 		/** Expands the specified asset category, if specified. If not, it will expand Basic/Common instead. */
 		SLATE_ARGUMENT(TOptional<FAssetCategoryPath>, DefaultMenuExpansionCategory)
 
+		/** If true, adds custom categories as sections (expanded) vs. as sub-menus */
+		SLATE_ARGUMENT(bool, bUseSectionsForCustomCategories)
+
 	SLATE_END_ARGS()
 
 	/** Constructs this widget with InArgs */
diff --git a/Engine/Source/Editor/ContentBrowser/Public/IContentBrowserSingleton.h b/Engine/Source/Editor/ContentBrowser/Public/IContentBrowserSingleton.h
--- a/Engine/Source/Editor/ContentBrowser/Public/IContentBrowserSingleton.h
+++ b/Engine/Source/Editor/ContentBrowser/Public/IContentBrowserSingleton.h
@@ -303,6 +303,9 @@ struct FAssetPickerConfig
 	/** Override the default filter context menu layout */
 	EAssetTypeCategories::Type DefaultFilterMenuExpansion;
 
+	/** If we display filters & set to true, we will add sections instead of sub-menus for other filters. Useful if the number of additional filters is small. */
+	bool bUseSectionsForCustomFilterCategories;
+
 	FAssetPickerConfig()
 		: SelectionMode( ESelectionMode::Multi )
 		, ThumbnailLabel( EThumbnailLabel::ClassName )
@@ -325,6 +328,7 @@ struct FAssetPickerConfig
 		, bShowTypeInColumnView(true)
 		, bSortByPathInColumnView(false)
 		, DefaultFilterMenuExpansion(EAssetTypeCategories::Basic)
+		, bUseSectionsForCustomFilterCategories(false)
 	{}
 };
 
diff --git a/Engine/Source/Editor/EditorWidgets/Public/Filters/SAssetFilterBar.h b/Engine/Source/Editor/EditorWidgets/Public/Filters/SAssetFilterBar.h
--- a/Engine/Source/Editor/EditorWidgets/Public/Filters/SAssetFilterBar.h
+++ b/Engine/Source/Editor/EditorWidgets/Public/Filters/SAssetFilterBar.h
@@ -155,6 +155,7 @@ public:
 		, _CanChangeOrientation(false)
 		, _FilterPillStyle(EFilterPillStyle::Default)
 		, _DefaultMenuExpansionCategory(EAssetCategoryPaths::Basic)
+		, _bUseSectionsForCustomCategories(false)
 		{
 		
 		}
@@ -202,6 +203,9 @@ public:
 
 		/** Expands the specified asset category, if specified. If not, it will expand Basic/Common instead. */
 		SLATE_ARGUMENT(TOptional<FAssetCategoryPath>, DefaultMenuExpansionCategory)
+
+		/** If true, adds custom categories as sections (expanded) vs. as sub-menus */
+		SLATE_ARGUMENT(bool, bUseSectionsForCustomCategories)
 	
 	SLATE_END_ARGS()
 
@@ -224,6 +228,7 @@ public:
 		Args._FilterBarLayout = InArgs._FilterBarLayout;
 		Args._CanChangeOrientation = InArgs._CanChangeOrientation;
 		Args._FilterPillStyle = InArgs._FilterPillStyle;
+		Args._UseSectionsForCategories = InArgs._bUseSectionsForCustomCategories;
 		
 		SBasicFilterBar<FilterType>::Construct(Args);
 
@@ -1083,7 +1088,7 @@ protected:
 				
 				// If we are doing a full menu (i.e expanding basic) we add a menu entry which toggles all other categories
                 Section.AddMenuEntry(
-                	NAME_None,
+                	FName(FText::AsCultureInvariant(ExpandedCategory->SectionHeading).ToString()),
                 	ExpandedCategory->SectionHeading,
                 	MenuExpansion->Tooltip,
                 	FSlateIcon(FAppStyle::Get().GetStyleSetName(), "PlacementBrowser.Icons.Basic"),
