Added culture invariant identifier to expanded category section

diff --git a/Engine/Source/Editor/EditorWidgets/Public/Filters/SAssetFilterBar.h b/Engine/Source/Editor/EditorWidgets/Public/Filters/SAssetFilterBar.h
--- a/Engine/Source/Editor/EditorWidgets/Public/Filters/SAssetFilterBar.h
+++ b/Engine/Source/Editor/EditorWidgets/Public/Filters/SAssetFilterBar.h
@@ -1041,8 +1041,8 @@ protected:
 					// Otherwise create a new FCategoryMenu for the category and add it to the map
 					else
 					{
-						const FName ExtensionPoint = NAME_None;
 						const FText SectionHeading = FText::Format(LOCTEXT("WildcardFilterHeadingHeadingTooltip", "{0} Filters"), Category->Title);
+						const FName ExtensionPoint = FName(FText::AsCultureInvariant(SectionHeading).ToString());
 
 						FCategoryMenu NewCategoryMenu(ExtensionPoint, SectionHeading);
 						NewCategoryMenu.Classes.Add(CustomClassFilter);
