Sequencer: Fix transport control recording icon

#rnx

diff --git a/Engine/Source/Editor/EditorWidgets/Private/STransportControl.cpp b/Engine/Source/Editor/EditorWidgets/Private/STransportControl.cpp
--- a/Engine/Source/Editor/EditorWidgets/Private/STransportControl.cpp
+++ b/Engine/Source/Editor/EditorWidgets/Private/STransportControl.cpp
@@ -80,8 +80,22 @@ TSharedPtr<SWidget> STransportControl::MakeTransportControlWidget(ETransportCont
 			. IsFocusable(bAreButtonsFocusable)
 			[
 				SNew(SImage)
-				.ColorAndOpacity(FSlateColor::UseSubduedForeground())
-				.Image(this, &STransportControl::GetRecordStatusIcon )
+				.Image(FAppStyle::Get().GetBrush("Animation.Record"))
+				.ColorAndOpacity_Lambda([this]()
+				{
+					bool bIsRecording = false;
+					if (TransportControlArgs.OnGetRecording.IsBound())
+					{
+						bIsRecording = TransportControlArgs.OnGetRecording.Execute();
+					}
+
+					if (bIsRecording)
+					{
+						return FSlateColor::UseForeground();
+					}
+
+					return FSlateColor::UseSubduedForeground();
+				})
 			];
 	case ETransportControlWidgetType::ForwardPlay:
 		return SAssignNew(ForwardPlayButton, SButton)
diff --git a/Engine/Source/Editor/Sequencer/Private/Sequencer.cpp b/Engine/Source/Editor/Sequencer/Private/Sequencer.cpp
--- a/Engine/Source/Editor/Sequencer/Private/Sequencer.cpp
+++ b/Engine/Source/Editor/Sequencer/Private/Sequencer.cpp
@@ -4424,10 +4424,9 @@ TSharedRef<SWidget> FSequencer::OnCreateTransportRecord()
 				{
 					RecordingAnimation.Play(SequencerWidget.ToSharedRef(), true);
 				}
-
-				FLinearColor Color = FSlateColor::UseSubduedForeground().GetSpecifiedColor();
-				Color.A = 0.2f + 0.8f * RecordingAnimation.GetLerp();
-				return FSlateColor(Color);
+				
+				FLinearColor Color = FLinearColor::White;
+				return FSlateColor(Color.CopyWithNewOpacity(0.2f + 0.8f * RecordingAnimation.GetLerp()));
 			}
 
 			RecordingAnimation.Pause();
