Localize Checkbox works with multiple values.

diff --git a/Engine/Source/Editor/EditorWidgets/Private/STextPropertyEditableTextBox.cpp b/Engine/Source/Editor/EditorWidgets/Private/STextPropertyEditableTextBox.cpp
--- a/Engine/Source/Editor/EditorWidgets/Private/STextPropertyEditableTextBox.cpp
+++ b/Engine/Source/Editor/EditorWidgets/Private/STextPropertyEditableTextBox.cpp
@@ -879,11 +879,15 @@ void STextPropertyEditableTextBox::Construct(const FArguments& InArgs, const TSh
 
 bool STextPropertyEditableTextBox::IsTextLocalizable() const
 {
+	// All text need !IsCultureInvariant()
 	const int32 NumTexts = EditableTextProperty->GetNumTexts();
-	if (NumTexts == 1)
+	for (int32 Index = 0; Index < NumTexts; ++Index)
 	{
-		const FText PropertyValue = EditableTextProperty->GetText(0);
-		return !PropertyValue.IsCultureInvariant();
+		const FText PropertyValue = EditableTextProperty->GetText(Index);
+		if (PropertyValue.IsCultureInvariant())
+		{
+			return false;
+		}
 	}
 	return true;
 }
@@ -933,9 +937,9 @@ bool STextPropertyEditableTextBox::IsSourceTextReadOnly() const
 
 	// We can't edit the source string of string table references
 	const int32 NumTexts = EditableTextProperty->GetNumTexts();
-	if (NumTexts == 1)
+	for (int32 TextIndex = 0; TextIndex < NumTexts; ++TextIndex)
 	{
-		const FText TextValue = EditableTextProperty->GetText(0);
+		const FText TextValue = EditableTextProperty->GetText(TextIndex);
 		if (TextValue.IsFromStringTable())
 		{
 			return true;
@@ -954,9 +958,9 @@ bool STextPropertyEditableTextBox::IsIdentityReadOnly() const
 
 	// We can't edit the identity of texts that don't gather for localization
 	const int32 NumTexts = EditableTextProperty->GetNumTexts();
-	if (NumTexts == 1)
+	for (int32 TextIndex = 0; TextIndex < NumTexts; ++TextIndex)
 	{
-		const FText TextValue = EditableTextProperty->GetText(0);
+		const FText TextValue = EditableTextProperty->GetText(TextIndex);
 		if (!TextValue.ShouldGatherForLocalization())
 		{
 			return true;
@@ -1422,16 +1426,23 @@ FText STextPropertyEditableTextBox::GetPackageValue() const
 
 ECheckBoxState STextPropertyEditableTextBox::GetLocalizableCheckState() const
 {
+	TOptional<ECheckBoxState> Result;
+
 	const int32 NumTexts = EditableTextProperty->GetNumTexts();
-	if (NumTexts == 1)
+	for (int32 Index = 0; Index < NumTexts; ++Index)
 	{
-		const FText PropertyValue = EditableTextProperty->GetText(0);
+		const FText PropertyValue = EditableTextProperty->GetText(Index);
 
 		const bool bIsLocalized = !PropertyValue.IsCultureInvariant();
-		return bIsLocalized ? ECheckBoxState::Checked : ECheckBoxState::Unchecked;
+		ECheckBoxState NewState = bIsLocalized ? ECheckBoxState::Checked : ECheckBoxState::Unchecked;
+		if (NewState != Result.Get(NewState))
+		{
+			return ECheckBoxState::Undetermined;
+		}
+		Result = NewState;
 	}
 
-	return ECheckBoxState::Unchecked;
+	return Result.Get(ECheckBoxState::Unchecked);
 }
 
 void STextPropertyEditableTextBox::HandleLocalizableCheckStateChanged(ECheckBoxState InCheckboxState)
