Fix crash if SAssetFilter destroyed and AssetDefinitionRegistry version changes
* SAssetFilterBar doesn't clean up the registered delegate, leaving a stale this pointer

diff --git a/Engine/Source/Editor/EditorWidgets/Public/Filters/SAssetFilterBar.h b/Engine/Source/Editor/EditorWidgets/Public/Filters/SAssetFilterBar.h
--- a/Engine/Source/Editor/EditorWidgets/Public/Filters/SAssetFilterBar.h
+++ b/Engine/Source/Editor/EditorWidgets/Public/Filters/SAssetFilterBar.h
@@ -209,6 +209,8 @@ public:
 	
 	SLATE_END_ARGS()
 
+	virtual ~SAssetFilterBar() override;
+
 	/** Constructs this widget with InArgs */
 	void Construct( const FArguments& InArgs )
 	{
@@ -242,7 +244,7 @@ public:
 
 		// Re-create the Asset Type Action filters whenever the permission list changes
 		AssetClassPermissionList->OnFilterChanged().AddSP(this, &SAssetFilterBar<FilterType>::CreateAssetTypeActionFilters);
-		UAssetDefinitionRegistry::Get()->OnAssetDefinitionRegistryVersionChange().AddLambda([this](UAssetDefinitionRegistry*) { CreateAssetTypeActionFilters(); });
+		AssetRegistryVersionChangeHandle = UAssetDefinitionRegistry::Get()->OnAssetDefinitionRegistryVersionChange().AddLambda([this](UAssetDefinitionRegistry*) { CreateAssetTypeActionFilters(); });
 	}
 
 	virtual void SaveSettings()
@@ -1272,6 +1274,17 @@ protected:
 private:
 	/** The filter menu category to expand. */
 	TOptional<FAssetCategoryPath> DefaultMenuExpansionCategory;
+
+	FDelegateHandle AssetRegistryVersionChangeHandle;
 };
 
-#undef LOCTEXT_NAMESPACE
\ No newline at end of file
+template <typename FilterType>
+SAssetFilterBar<FilterType>::~SAssetFilterBar()
+{
+	if (UAssetDefinitionRegistry* AssetDefinitionRegistry = UAssetDefinitionRegistry::Get())
+	{
+		AssetDefinitionRegistry->OnAssetDefinitionRegistryVersionChange().Remove(AssetRegistryVersionChangeHandle);
+	}
+}
+
+#undef LOCTEXT_NAMESPACE
