Improve "Multiple Values" editing for FText properties

#rnx

diff --git a/Engine/Source/Editor/EditorWidgets/Private/STextPropertyEditableTextBox.cpp b/Engine/Source/Editor/EditorWidgets/Private/STextPropertyEditableTextBox.cpp
--- a/Engine/Source/Editor/EditorWidgets/Private/STextPropertyEditableTextBox.cpp
+++ b/Engine/Source/Editor/EditorWidgets/Private/STextPropertyEditableTextBox.cpp
@@ -1057,30 +1057,34 @@ FText STextPropertyEditableTextBox::GetToolTipText() const
 
 FText STextPropertyEditableTextBox::GetTextValue() const
 {
-	FText TextValue;
-
 	const int32 NumTexts = EditableTextProperty->GetNumTexts();
-	if (NumTexts == 1)
-	{
-		TextValue = EditableTextProperty->GetText(0);
 
-		if (const FString* SourceString = FTextInspector::GetSourceString(TextValue);
-			SourceString && !FTextLocalizationManager::Get().IsLocalizationLocked())
+	if (NumTexts > 0)
+	{
+		FString SourceString;
+		for (int32 TextIndex = 0; TextIndex < NumTexts; ++TextIndex)
 		{
-			// We should always edit the source string, but if the source string matches the current 
-			// display string then we can avoid making a temporary text from the source string
-			if (!SourceString->IsEmpty() && !SourceString->Equals(TextValue.ToString(), ESearchCase::CaseSensitive))
+			const FText TextValue = EditableTextProperty->GetText(TextIndex);
+			if (const FString* SourceStringPtr = FTextInspector::GetSourceString(TextValue))
+			{
+				if (TextIndex == 0)
+				{
+					SourceString = *SourceStringPtr;
+				}
+				else if (!SourceString.Equals(*SourceStringPtr, ESearchCase::CaseSensitive))
+				{
+					return MultipleValuesText;
+				}
+			}
+			else if (TextIndex > 0)
 			{
-				TextValue = FText::AsCultureInvariant(*SourceString);
+				return MultipleValuesText;
 			}
 		}
-	}
-	else if (NumTexts > 1)
-	{
-		TextValue = MultipleValuesText;
+		return FText::AsCultureInvariant(MoveTemp(SourceString));
 	}
 
-	return TextValue;
+	return FText::GetEmpty();
 }
 
 void STextPropertyEditableTextBox::OnTextChanged(const FText& NewText)
@@ -1197,24 +1201,35 @@ void STextPropertyEditableTextBox::SetTextError(const FText& InErrorMsg)
 
 FText STextPropertyEditableTextBox::GetNamespaceValue() const
 {
-	FText NamespaceValue;
-
 	const int32 NumTexts = EditableTextProperty->GetNumTexts();
-	if (NumTexts == 1)
+
+	if (NumTexts > 0)
 	{
-		const FText PropertyValue = EditableTextProperty->GetText(0);
-		TOptional<FString> FoundNamespace = FTextInspector::GetNamespace(PropertyValue);
-		if (FoundNamespace.IsSet())
+		FString NamespaceValue;
+		for (int32 TextIndex = 0; TextIndex < NumTexts; ++TextIndex)
 		{
-			NamespaceValue = FText::FromString(TextNamespaceUtil::StripPackageNamespace(FoundNamespace.GetValue()));
+			const FText TextValue = EditableTextProperty->GetText(TextIndex);
+			if (TOptional<FString> FoundNamespace = FTextInspector::GetNamespace(TextValue))
+			{
+				FString CleanNamespace = TextNamespaceUtil::StripPackageNamespace(FoundNamespace.GetValue());
+				if (TextIndex == 0)
+				{
+					NamespaceValue = MoveTemp(CleanNamespace);
+				}
+				else if (!NamespaceValue.Equals(CleanNamespace, ESearchCase::CaseSensitive))
+				{
+					return MultipleValuesText;
+				}
+			}
+			else if (TextIndex > 0)
+			{
+				return MultipleValuesText;
+			}
 		}
-	}
-	else if (NumTexts > 1)
-	{
-		NamespaceValue = MultipleValuesText;
+		return FText::AsCultureInvariant(MoveTemp(NamespaceValue));
 	}
 
-	return NamespaceValue;
+	return FText::GetEmpty();
 }
 
 void STextPropertyEditableTextBox::OnNamespaceChanged(const FText& NewText)
@@ -1289,24 +1304,34 @@ void STextPropertyEditableTextBox::OnNamespaceCommitted(const FText& NewText, ET
 
 FText STextPropertyEditableTextBox::GetKeyValue() const
 {
-	FText KeyValue;
-
 	const int32 NumTexts = EditableTextProperty->GetNumTexts();
-	if (NumTexts == 1)
+
+	if (NumTexts > 0)
 	{
-		const FText PropertyValue = EditableTextProperty->GetText(0);
-		TOptional<FString> FoundKey = FTextInspector::GetKey(PropertyValue);
-		if (FoundKey.IsSet())
+		FString KeyValue;
+		for (int32 TextIndex = 0; TextIndex < NumTexts; ++TextIndex)
 		{
-			KeyValue = FText::FromString(FoundKey.GetValue());
+			const FText TextValue = EditableTextProperty->GetText(TextIndex);
+			if (TOptional<FString> FoundKey = FTextInspector::GetKey(TextValue))
+			{
+				if (TextIndex == 0)
+				{
+					KeyValue = FoundKey.GetValue();
+				}
+				else if (!KeyValue.Equals(FoundKey.GetValue(), ESearchCase::CaseSensitive))
+				{
+					return MultipleValuesText;
+				}
+			}
+			else if (TextIndex > 0)
+			{
+				return MultipleValuesText;
+			}
 		}
-	}
-	else if (NumTexts > 1)
-	{
-		KeyValue = MultipleValuesText;
+		return FText::AsCultureInvariant(MoveTemp(KeyValue));
 	}
 
-	return KeyValue;
+	return FText::GetEmpty();
 }
 
 #if USE_STABLE_LOCALIZATION_KEYS
@@ -1402,24 +1427,35 @@ void STextPropertyEditableTextBox::OnKeyCommitted(const FText& NewText, ETextCom
 
 FText STextPropertyEditableTextBox::GetPackageValue() const
 {
-	FText PackageValue;
-
 	const int32 NumTexts = EditableTextProperty->GetNumTexts();
-	if (NumTexts == 1)
+
+	if (NumTexts > 0)
 	{
-		const FText PropertyValue = EditableTextProperty->GetText(0);
-		TOptional<FString> FoundNamespace = FTextInspector::GetNamespace(PropertyValue);
-		if (FoundNamespace.IsSet())
+		FString PackageValue;
+		for (int32 TextIndex = 0; TextIndex < NumTexts; ++TextIndex)
 		{
-			PackageValue = FText::FromString(TextNamespaceUtil::ExtractPackageNamespace(FoundNamespace.GetValue()));
+			const FText TextValue = EditableTextProperty->GetText(TextIndex);
+			if (TOptional<FString> FoundNamespace = FTextInspector::GetNamespace(TextValue))
+			{
+				FString PackageNamespace = TextNamespaceUtil::ExtractPackageNamespace(FoundNamespace.GetValue());
+				if (TextIndex == 0)
+				{
+					PackageValue = MoveTemp(PackageNamespace);
+				}
+				else if (!PackageValue.Equals(PackageNamespace, ESearchCase::CaseSensitive))
+				{
+					return MultipleValuesText;
+				}
+			}
+			else if (TextIndex > 0)
+			{
+				return MultipleValuesText;
+			}
 		}
-	}
-	else if (NumTexts > 1)
-	{
-		PackageValue = MultipleValuesText;
+		return FText::AsCultureInvariant(MoveTemp(PackageValue));
 	}
 
-	return PackageValue;
+	return FText::GetEmpty();
 }
 
 #endif // USE_STABLE_LOCALIZATION_KEYS
