Implemented the concept of "internal" string tables

These are tables that shouldn't appear in any pickers, and also allow editing the source string of a referencing FText to "auto break" the FText's link to the string table. This is intended to be used for string tables that are automatically generated and referenced by custom code.

#rnx

diff --git a/Engine/Source/Editor/EditorWidgets/Private/STextPropertyEditableTextBox.cpp b/Engine/Source/Editor/EditorWidgets/Private/STextPropertyEditableTextBox.cpp
--- a/Engine/Source/Editor/EditorWidgets/Private/STextPropertyEditableTextBox.cpp
+++ b/Engine/Source/Editor/EditorWidgets/Private/STextPropertyEditableTextBox.cpp
@@ -345,6 +345,11 @@ void STextPropertyEditableStringTableReference::UpdateStringTableComboOptions()
 			FStringTableConstPtr StringTable = FStringTableRegistry::Get().FindStringTable(StringTableId);
 			if (StringTable.IsValid())
 			{
+				if (StringTable->IsInternal())
+				{
+					continue; // can't pick internal string tables
+				}
+
 				StringTable->EnumerateSourceStrings([&](const FString& InKey, const FString& InSourceString) -> bool
 				{
 					HasEntries = true;
@@ -379,6 +384,11 @@ void STextPropertyEditableStringTableReference::UpdateStringTableComboOptions()
 	// Process the remaining non-asset string tables now
 	FStringTableRegistry::Get().EnumerateStringTables([&](const FName& InTableId, const FStringTableConstRef& InStringTable) -> bool
 	{
+		if (InStringTable->IsInternal())
+		{
+			return true; // can't pick internal string tables
+		}
+
 		const bool bAlreadyAdded = StringTableComboOptions.ContainsByPredicate([InTableId](const TSharedPtr<FAvailableStringTable>& InAvailableStringTable)
 		{
 			return InAvailableStringTable->TableId == InTableId;
@@ -935,14 +945,22 @@ bool STextPropertyEditableTextBox::IsSourceTextReadOnly() const
 		return true;
 	}
 
-	// We can't edit the source string of string table references
+	// We can't edit the source string of external string table references
 	const int32 NumTexts = EditableTextProperty->GetNumTexts();
 	for (int32 TextIndex = 0; TextIndex < NumTexts; ++TextIndex)
 	{
 		const FText TextValue = EditableTextProperty->GetText(TextIndex);
 		if (TextValue.IsFromStringTable())
 		{
-			return true;
+			FName TableId;
+			FString Key;
+			FTextInspector::GetTableIdAndKey(TextValue, TableId, Key);
+
+			if (FStringTableConstPtr StringTable = FStringTableRegistry::Get().FindStringTable(TableId);
+				StringTable && !StringTable->IsInternal())
+			{
+				return true;
+			}
 		}
 	}
 
diff --git a/Engine/Source/Runtime/Core/Private/Internationalization/StringTableCore.cpp b/Engine/Source/Runtime/Core/Private/Internationalization/StringTableCore.cpp
--- a/Engine/Source/Runtime/Core/Private/Internationalization/StringTableCore.cpp
+++ b/Engine/Source/Runtime/Core/Private/Internationalization/StringTableCore.cpp
@@ -91,6 +91,7 @@ const FString& FStringTableEntry::GetPlaceholderSourceString()
 FStringTable::FStringTable()
 	: OwnerAsset(nullptr)
 	, bIsLoaded(true)
+	, bIsInternal(false)
 {
 }
 
@@ -120,6 +121,16 @@ void FStringTable::IsLoaded(const bool bInIsLoaded)
 	bIsLoaded = bInIsLoaded;
 }
 
+bool FStringTable::IsInternal() const
+{
+	return bIsInternal;
+}
+
+void FStringTable::IsInternal(const bool bInIsInternal)
+{
+	bIsInternal = bInIsInternal;
+}
+
 FString FStringTable::GetNamespace() const
 {
 	return TableNamespace.GetChars();
diff --git a/Engine/Source/Runtime/Core/Public/Internationalization/StringTableCore.h b/Engine/Source/Runtime/Core/Public/Internationalization/StringTableCore.h
--- a/Engine/Source/Runtime/Core/Public/Internationalization/StringTableCore.h
+++ b/Engine/Source/Runtime/Core/Public/Internationalization/StringTableCore.h
@@ -101,9 +101,15 @@ public:
 	/** Has this string table been fully loaded yet? (used during asset loading) */
 	CORE_API bool IsLoaded() const;
 
-	/** Set whether this string table has been fully loaded yet */
+	/** Set whether this string table has been fully loaded */
 	CORE_API void IsLoaded(const bool bInIsLoaded);
 
+	/** Is this string table considered internal (ie, shouldn't appear in pickers) */
+	CORE_API bool IsInternal() const;
+
+	/** Set whether this string table is considered internal */
+	CORE_API void IsInternal(const bool bInIsInternal);
+
 	/** @return The namespace used by all entries in this string table */
 	CORE_API FString GetNamespace() const;
 
@@ -167,6 +173,9 @@ private:
 	/** True if this table has been fully loaded (used for assets) */
 	bool bIsLoaded;
 
+	/** True if this table is internal and shouldn't appear in pickers */
+	bool bIsInternal;
+
 	/** The namespace to use for all the strings in this table */
 	FTextKey TableNamespace;
 
