[ContentBrowser] Re-arranged filter menu (new style), reviewed in a previous CL but not submitted

diff --git a/Engine/Source/Editor/EditorWidgets/Public/Filters/FilterBarConfig.h b/Engine/Source/Editor/EditorWidgets/Public/Filters/FilterBarConfig.h
--- a/Engine/Source/Editor/EditorWidgets/Public/Filters/FilterBarConfig.h
+++ b/Engine/Source/Editor/EditorWidgets/Public/Filters/FilterBarConfig.h
@@ -62,7 +62,11 @@ public:
 
 	UPROPERTY()
 	EFilterBarLayout FilterBarLayout = EFilterBarLayout::Horizontal;
-	
+
+	/** Whether to filter paths recursively. Also @see UContentBrowserSettings::FilterRecursively */
+	UPROPERTY()
+	bool bFilterRecursively = true;
+
 	void Empty()
 	{
 		CustomFilters.Empty();
diff --git a/Engine/Source/Editor/EditorWidgets/Public/Filters/SAssetFilterBar.h b/Engine/Source/Editor/EditorWidgets/Public/Filters/SAssetFilterBar.h
--- a/Engine/Source/Editor/EditorWidgets/Public/Filters/SAssetFilterBar.h
+++ b/Engine/Source/Editor/EditorWidgets/Public/Filters/SAssetFilterBar.h
@@ -11,6 +11,7 @@
 #include "Containers/Set.h"
 #include "Containers/UnrealString.h"
 #include "Delegates/Delegate.h"
+#include "Factories.h"
 #include "Filters/CustomClassFilterData.h"
 #include "Filters/CustomTextFilters.h"
 #include "Filters/FilterBarConfig.h"
@@ -18,7 +19,9 @@
 #include "Filters/SBasicFilterBar.h"
 #include "Framework/Commands/UIAction.h"
 #include "Framework/Commands/UICommandInfo.h"
+#include "HAL/IConsoleManager.h"
 #include "HAL/Platform.h"
+#include "HAL/PlatformApplicationMisc.h"
 #include "IAssetTools.h"
 #include "IAssetTypeActions.h"
 #include "Internationalization/Internationalization.h"
@@ -28,6 +31,7 @@
 #include "Math/Color.h"
 #include "Misc/AssertionMacros.h"
 #include "Misc/Attribute.h"
+#include "Misc/CoreMiscDefines.h"
 #include "Misc/NamePermissionList.h"
 #include "Modules/ModuleManager.h"
 #include "SlateGlobals.h"
@@ -58,6 +62,26 @@
 class SFilterSearchBox;
 class SWidget;
 
+namespace UE::Editor::EditorWidgets
+{
+	static bool IsNewStyleEnabled()
+	{
+		// @note: This is tied to the ContentBrowser new style (5.6+)
+
+		static bool bIsNewStyleEnabled = false;
+		UE_CALL_ONCE([&]()
+		{
+			if (const IConsoleVariable* CVar = IConsoleManager::Get().FindConsoleVariable(TEXT("ContentBrowser.EnableNewStyle")))
+			{
+				ensureAlwaysMsgf(!EnumHasAnyFlags(CVar->GetFlags(), ECVF_Default), TEXT("The CVar should have already been set from commandline, @see: UnrealEdGlobals.cpp, UE::Editor::ContentBrowser::EnableContentBrowserNewStyleCVarRegistration."));
+				bIsNewStyleEnabled = CVar->GetBool();
+			}
+		});
+
+		return bIsNewStyleEnabled;
+	}
+}
+
 #define LOCTEXT_NAMESPACE "FilterBar"
 
 /** Delegate that subclasses can use to specify classes to not include in this filter
@@ -268,6 +292,11 @@ public:
 		if(this->bCanChangeOrientation)
 		{
 			FilterBarConfig->FilterBarLayout = this->FilterBarLayout;
+
+			if (UE::Editor::EditorWidgets::IsNewStyleEnabled())
+			{
+				FilterBarConfig->bFilterRecursively = this->bFilterPathsRecursively;
+			}
 		}
 
 		SaveConfig();
@@ -296,6 +325,10 @@ public:
 		if(FilterBarConfig->bIsLayoutSaved)
 		{
 			this->SetFilterLayout(FilterBarConfig->FilterBarLayout);
+			if (UE::Editor::EditorWidgets::IsNewStyleEnabled())
+			{
+				this->SetFilterPathsRecursively(FilterBarConfig->bFilterRecursively);
+			}
 		}
 
 		this->OnFilterChanged.ExecuteIfBound();
@@ -637,6 +670,11 @@ public:
 			}
 		}
 
+		if (UE::Editor::EditorWidgets::IsNewStyleEnabled())
+		{
+			CombinedFilter.bRecursivePaths = this->bFilterPathsRecursively;
+		}
+
 		return CombinedFilter;
 	}
 
@@ -932,7 +970,7 @@ protected:
 		for (int32 ClassIdx = 0; ClassIdx < CustomClassFilterDatas.Num(); ++ClassIdx)
 		{
 			const TSharedPtr<FCustomClassFilterData>& CustomClassFilterData = CustomClassFilterDatas[ClassIdx];
-			
+
 			const FText& LabelText = CustomClassFilterData->GetName();
 			Section.AddMenuEntry(
 				NAME_None,
@@ -947,7 +985,7 @@ protected:
 				);
 		}
 	}
-	
+
 	void CreateFiltersMenuCategory(UToolMenu* InMenu, const TArray<TSharedPtr<FCustomClassFilterData>> CustomClassFilterDatas) const
 	{
 		CreateFiltersMenuCategory(InMenu->AddSection("Section"), CustomClassFilterDatas);
@@ -988,7 +1026,7 @@ protected:
 				}
 			}));
 		}
-		
+
 		FToolMenuContext ToolMenuContext(CreateAssetFilterBarContext());
 		return UToolMenus::Get()->GenerateWidget(FilterMenuName, ToolMenuContext);
 	} 
@@ -1048,7 +1086,7 @@ protected:
 
 						FCategoryMenu NewCategoryMenu(ExtensionPoint, SectionHeading);
 						NewCategoryMenu.Classes.Add(CustomClassFilter);
-					
+
 						CategoryToMenuMap.Add(Category, NewCategoryMenu);
 					}
 				}
@@ -1076,10 +1114,112 @@ protected:
 		// Populate the common filter sections (Reset Filters etc)
 		{
 			this->PopulateCommonFilterSections(Menu);
+
+			if (UE::Editor::EditorWidgets::IsNewStyleEnabled())
+			{
+				// Remove legacy section
+				Menu->RemoveSection("FilterBarTextFilters");
+
+				FToolMenuSection& AdvancedSection = Menu->FindOrAddSection(
+					"FilterBarAdvanced",
+					LOCTEXT("FilterBarAdvancedSection", "Advanced"),
+					FToolMenuInsert("BasicFilterBarFiltersMenu", EToolMenuInsertType::Last));
+
+				AdvancedSection.AddSubMenu(
+					"CustomFiltersSubMenu",
+					LOCTEXT("FilterBarTextFilters", "Custom Filters"),
+					LOCTEXT("FilterBarTextFiltersTooltip", "Custom Filters"),
+					FNewToolMenuDelegate::CreateSP(this, &SAssetFilterBar<FilterType>::CreateTextFiltersMenu),
+					false,
+					FSlateIcon(FAppStyle::Get().GetStyleSetName(), "Icons.Filter"));
+
+				if (FToolMenuSection* ManageSection = Menu->FindSection("FilterBarResetFilters"))
+				{
+					ManageSection->Label = LOCTEXT("FilterBarManageSection", "Manage");
+
+					ManageSection->AddMenuEntry(
+						"CopyFilters",
+						LOCTEXT("FilterListCopyFilters", "Copy Filters"),
+						LOCTEXT("FilterListCopyFiltersTooltip", "Copy the current filter selection"),
+						FSlateIcon(FAppStyle::Get().GetStyleSetName(), "GenericCommands.Copy"),
+						FUIAction(
+							FExecuteAction::CreateSP(this, &SAssetFilterBar<FilterType>::OnCopyFilters))
+					);
+
+					ManageSection->AddMenuEntry(
+						"PasteFilters",
+						LOCTEXT("FilterListPasteFilters", "Paste Filters"),
+						LOCTEXT("FilterListPasteFiltersTooltip", "Paste to the current filter selection"),
+						FSlateIcon(FAppStyle::Get().GetStyleSetName(), "GenericCommands.Paste"),
+						FUIAction(
+							FExecuteAction::CreateSP(this, &SAssetFilterBar<FilterType>::OnPasteFilters),
+							FCanExecuteAction::CreateSP(this, &SAssetFilterBar<FilterType>::CanPasteFilters))
+					);
+
+					ManageSection->AddMenuEntry(
+						"FilterRecursively",
+						LOCTEXT("FilterListFilterRecursively", "Filter Recursively"),
+						LOCTEXT("FilterListFilterRecursivelyTooltip", "Apply the current filter selection recursively, relevant to the current path"),
+						{ },
+						FUIAction(
+							FExecuteAction::CreateSP(this, &SAssetFilterBar<FilterType>::OnFilterPathsRecursively),
+							FCanExecuteAction(),
+							FIsActionChecked::CreateSP(this, &SAssetFilterBar<FilterType>::IsFilteringPathsRecursively)),
+						EUserInterfaceActionType::Check
+					);
+				}
+
+				FToolMenuSection& DisplaySection = Menu->AddSection(
+					"FilterBarDisplay",
+					LOCTEXT("FilterListDisplaySection", "Filter Display"));
+				{
+					DisplaySection.AddMenuEntry(
+						"HorizontalLayout",
+						LOCTEXT("FilterListHorizontalLayout", "Horizontal"),
+						LOCTEXT("FilterListHorizontalLayoutToolTip", "Swap to a Horizontal layout for the filter bar"),
+						FSlateIcon(),
+						FUIAction(
+							FExecuteAction::CreateSPLambda(
+							this, [this]
+							{
+								if (this->GetFilterLayout() != EFilterBarLayout::Horizontal)
+								{
+									this->SetFilterLayout(EFilterBarLayout::Horizontal);
+								}
+							}),
+							FCanExecuteAction(),
+							FIsActionChecked::CreateSPLambda([this]
+							{
+								return this->GetFilterLayout() == EFilterBarLayout::Horizontal;
+							})),
+						EUserInterfaceActionType::RadioButton
+					);
+
+					DisplaySection.AddMenuEntry(
+						"VerticalLayout",
+						LOCTEXT("FilterListVerticalLayout", "Vertical"),
+						LOCTEXT("FilterListVerticalLayoutToolTip", "Swap to a vertical layout for the filter bar"),
+						FSlateIcon(),
+						FUIAction(
+							FExecuteAction::CreateSPLambda(this, [this]
+							{
+								if (this->GetFilterLayout() != EFilterBarLayout::Vertical)
+								{
+									this->SetFilterLayout(EFilterBarLayout::Vertical);
+								}
+							}),
+							FCanExecuteAction(),
+							FIsActionChecked::CreateSPLambda(this, [this]
+							{
+								return this->GetFilterLayout() == EFilterBarLayout::Vertical;
+							})),
+						EUserInterfaceActionType::RadioButton
+					);
+				}
+			}
 		}
 		
 		// If we want to expand a category
-		// if(MenuExpansion)
 		if(MenuExpansion)
 		{
 			// First add the expanded category, this appears as standard entries in the list (Note: intentionally not using FindChecked here as removing it from the map later would cause the ref to be garbage)
@@ -1087,38 +1227,76 @@ protected:
 			if(ExpandedCategory)
 			{
 				FToolMenuSection& Section = Menu->AddSection(ExpandedCategory->SectionExtensionHook, ExpandedCategory->SectionHeading);
-				
-				// If we are doing a full menu (i.e expanding basic) we add a menu entry which toggles all other categories
-                Section.AddMenuEntry(
-                	FName(FText::AsCultureInvariant(ExpandedCategory->SectionHeading).ToString()),
-                	ExpandedCategory->SectionHeading,
-                	MenuExpansion->Tooltip,
-                	FSlateIcon(FAppStyle::Get().GetStyleSetName(), "PlacementBrowser.Icons.Basic"),
-                	FUIAction(
-                	FExecuteAction::CreateSP( this, &SAssetFilterBar<FilterType>::FilterByTypeCategoryClicked, MenuExpansion, ExpandedCategory->Classes ),
-                	FCanExecuteAction(),
-                	FGetActionCheckState::CreateSP(this, &SAssetFilterBar<FilterType>::IsTypeCategoryChecked, MenuExpansion, ExpandedCategory->Classes ) ),
-                	EUserInterfaceActionType::ToggleButton
-                	);
-
-				Section.AddSeparator("ExpandedCategorySeparator");
-				
+
+				if (!UE::Editor::EditorWidgets::IsNewStyleEnabled())
+				{
+					// If we are doing a full menu (i.e expanding basic) we add a menu entry which toggles all other categories
+					Section.AddMenuEntry(
+						FName(FText::AsCultureInvariant(ExpandedCategory->SectionHeading).ToString()),
+						ExpandedCategory->SectionHeading,
+						MenuExpansion->Tooltip,
+						FSlateIcon(FAppStyle::Get().GetStyleSetName(), "PlacementBrowser.Icons.Basic"),
+						FUIAction(
+						FExecuteAction::CreateSP( this, &SAssetFilterBar<FilterType>::FilterByTypeCategoryClicked, MenuExpansion, ExpandedCategory->Classes ),
+						FCanExecuteAction(),
+						FGetActionCheckState::CreateSP(this, &SAssetFilterBar<FilterType>::IsTypeCategoryChecked, MenuExpansion, ExpandedCategory->Classes ) ),
+						EUserInterfaceActionType::ToggleButton
+						);
+
+					Section.AddSeparator("ExpandedCategorySeparator");
+				}
+
 				// Now populate with all the assets from the expanded category
-				SAssetFilterBar<FilterType>::CreateFiltersMenuCategory( Section, ExpandedCategory->Classes);
-				
+				SAssetFilterBar<FilterType>::CreateFiltersMenuCategory(Section, ExpandedCategory->Classes);
+
 				// Remove the Expanded from the map now, as this is treated differently and is no longer needed.
 				CategoryToMenuMap.Remove(MenuExpansion);
 			}
 		}
 
-		// Add all the other categories as submenus
-		FToolMenuSection& Section = Menu->AddSection("AssetFilterBarFilterAdvancedAsset", LOCTEXT("AdvancedAssetsMenuHeading", "Type Filters"));
-			
 		// Sort by category name so that we add the submenus in alphabetical order
 		CategoryToMenuMap.KeySort([](const TSharedPtr<FFilterCategory>& A, const TSharedPtr<FFilterCategory>& B) {
 			return A->Title.CompareTo(B->Title) < 0;
 		});
 
+if (UE::Editor::EditorWidgets::IsNewStyleEnabled())
+{
+		FToolMenuSection& AdvancedSection = Menu->FindOrAddSection("FilterBarAdvanced");
+
+		FToolMenuEntry& AllFiltersSubMenu = AdvancedSection.AddSubMenu(
+			"AllFilters",
+			LOCTEXT("AllAssetsMenuHeading", "All Filters"),
+			{ },
+			FNewToolMenuDelegate::CreateLambda([this, CategoryToMenuMap](UToolMenu* InMenu){
+				// Add all the other categories as submenus, un-labelled, acts only as a root for the submenus
+				FToolMenuSection& TypeFiltersSection = InMenu->AddSection("AssetFilterBarFilterAdvancedAsset", { });
+
+				for (const TPair<TSharedPtr<FFilterCategory>, FCategoryMenu>& CategoryMenuPair : CategoryToMenuMap)
+				{
+					TypeFiltersSection.AddSubMenu(
+						FName(FText::AsCultureInvariant(CategoryMenuPair.Key->Title).ToString()),
+						CategoryMenuPair.Key->Title,
+						CategoryMenuPair.Key->Tooltip,
+						FNewToolMenuDelegate::CreateSP(this, &SAssetFilterBar<FilterType>::CreateFiltersMenuCategory, CategoryMenuPair.Value.Classes),
+						FUIAction(
+						FExecuteAction::CreateSP(this, &SAssetFilterBar<FilterType>::FilterByTypeCategoryClicked, CategoryMenuPair.Key, CategoryMenuPair.Value.Classes),
+						FCanExecuteAction(),
+						FGetActionCheckState::CreateSP(this, &SAssetFilterBar<FilterType>::IsTypeCategoryChecked, CategoryMenuPair.Key, CategoryMenuPair.Value.Classes)),
+						EUserInterfaceActionType::ToggleButton
+					);
+				}
+
+				// Now add all non-asset filters
+				this->PopulateCustomFilters(InMenu);
+			}));
+
+		AllFiltersSubMenu.InsertPosition = FToolMenuInsert("CustomFiltersSubMenu", EToolMenuInsertType::After);
+}
+else
+{
+		// Add all the other categories as submenus
+		FToolMenuSection& Section = Menu->AddSection("AssetFilterBarFilterAdvancedAsset", LOCTEXT("AdvancedAssetsMenuHeading", "Type Filters"));
+
 		// Now actually add them
 		for (const TPair<TSharedPtr<FFilterCategory>, FCategoryMenu>& CategoryMenuPair : CategoryToMenuMap)
 		{
@@ -1137,6 +1315,7 @@ protected:
 
 		// Now add all non-asset filters
 		this->PopulateCustomFilters(Menu);
+}
 	}
 
 	/* Asset Type filter related functionality */
@@ -1220,7 +1399,6 @@ protected:
 			{
 				return ECheckBoxState::Undetermined;
 			}
-			
 		}
 
 		if (bIsAnyActionInUse)
@@ -1246,7 +1424,64 @@ protected:
  		// An asset type category is in use if any of its type actions are in use (ECheckBoxState::Checked or ECheckBoxState::Undetermined)
  		return true;
  	}
-	
+
+	void OnCopyFilters() const
+	{
+		if (const FFilterBarSettings* SettingsToCopyFrom = GetConstConfig())
+		{
+			FString Output = TEXT("");
+			FFilterBarSettings::StaticStruct()->ExportText(Output, SettingsToCopyFrom, nullptr, nullptr, (PPF_ExportsNotFullyQualified | PPF_Copy | PPF_Delimited | PPF_IncludeTransient), nullptr);
+			FPlatformApplicationMisc::ClipboardCopy(*Output);
+		}
+	}
+
+	void OnPasteFilters()
+	{
+		if (FFilterBarSettings* SettingsToPasteTo = GetMutableConfig())
+		{
+			FString Output = TEXT("");
+			FPlatformApplicationMisc::ClipboardPaste(Output);
+
+			FFilterBarSettings::StaticStruct()->ImportText(*Output, SettingsToPasteTo, nullptr, 0, GLog, GetPathNameSafe(FFilterBarSettings::StaticStruct()));
+
+			LoadSettings();
+		}
+	}
+
+	bool CanPasteFilters() const
+	{
+		FString Output = TEXT("");
+		FPlatformApplicationMisc::ClipboardPaste(Output);
+
+		FCustomizableTextObjectFactory Factory(GWarn);
+		if (!Factory.CanCreateObjectsFromText(Output))
+		{
+			return false;
+		}
+
+		FFilterBarSettings Value;
+		return FFilterBarSettings::StaticStruct()->ImportText(*Output, &Value, nullptr, 0, GLog, GetPathNameSafe(FFilterBarSettings::StaticStruct())) != nullptr;
+	}
+
+	void SetFilterPathsRecursively(const bool bInFilterRecursively)
+	{
+		if (bFilterPathsRecursively != bInFilterRecursively)
+		{
+			bFilterPathsRecursively = bInFilterRecursively;
+			this->OnFilterChanged.ExecuteIfBound();
+		}
+	}
+
+	void OnFilterPathsRecursively()
+	{
+		SetFilterPathsRecursively(!bFilterPathsRecursively);
+	}
+
+	bool IsFilteringPathsRecursively() const
+	{
+		return bFilterPathsRecursively;
+	}
+
 protected:
 	/** A copy of all AssetFilters in this->Filters for convenient access */
 	TArray< TSharedRef<SAssetFilter> > AssetFilters;
@@ -1271,6 +1506,8 @@ protected:
 	 *  bar. These are stored so we can save these when calling SaveSettings() so that they are not lost */
 	TMap<FString, bool> UnknownTypeFilters;
 
+	bool bFilterPathsRecursively = true;
+
 private:
 	/** The filter menu category to expand. */
 	TOptional<FAssetCategoryPath> DefaultMenuExpansionCategory;
