[ContentBrowser] AssetFilter cis fix

diff --git a/Engine/Source/Editor/EditorWidgets/Public/Filters/SAssetFilterBar.h b/Engine/Source/Editor/EditorWidgets/Public/Filters/SAssetFilterBar.h
--- a/Engine/Source/Editor/EditorWidgets/Public/Filters/SAssetFilterBar.h
+++ b/Engine/Source/Editor/EditorWidgets/Public/Filters/SAssetFilterBar.h
@@ -1188,7 +1188,7 @@ protected:
 								}
 							}),
 							FCanExecuteAction(),
-							FIsActionChecked::CreateSPLambda([this]
+							FIsActionChecked::CreateSPLambda(this, [this]
 							{
 								return this->GetFilterLayout() == EFilterBarLayout::Horizontal;
 							})),
