[ContentBrowser] Fixed Asset Filter conditional display of CreateTextFilter

diff --git a/Engine/Source/Editor/EditorWidgets/Public/Filters/SAssetFilterBar.h b/Engine/Source/Editor/EditorWidgets/Public/Filters/SAssetFilterBar.h
--- a/Engine/Source/Editor/EditorWidgets/Public/Filters/SAssetFilterBar.h
+++ b/Engine/Source/Editor/EditorWidgets/Public/Filters/SAssetFilterBar.h
@@ -1125,13 +1125,17 @@ protected:
 					LOCTEXT("FilterBarAdvancedSection", "Advanced"),
 					FToolMenuInsert("BasicFilterBarFiltersMenu", EToolMenuInsertType::Last));
 
-				AdvancedSection.AddSubMenu(
-					"CustomFiltersSubMenu",
-					LOCTEXT("FilterBarTextFilters", "Custom Filters"),
-					LOCTEXT("FilterBarTextFiltersTooltip", "Custom Filters"),
-					FNewToolMenuDelegate::CreateSP(this, &SAssetFilterBar<FilterType>::CreateTextFiltersMenu),
-					false,
-					FSlateIcon(FAppStyle::Get().GetStyleSetName(), "Icons.Filter"));
+				// Only add the custom text filter submenu if we have a valid CreateTextFilter delegate to use
+				if (this->CreateTextFilter.IsBound())
+				{
+					AdvancedSection.AddSubMenu(
+						"CustomFiltersSubMenu",
+						LOCTEXT("FilterBarTextFilters", "Custom Filters"),
+						LOCTEXT("FilterBarTextFiltersTooltip", "Custom Filters"),
+						FNewToolMenuDelegate::CreateSP(this, &SAssetFilterBar<FilterType>::CreateTextFiltersMenu),
+						false,
+						FSlateIcon(FAppStyle::Get().GetStyleSetName(), "Icons.Filter"));
+				}
 
 				if (FToolMenuSection* ManageSection = Menu->FindSection("FilterBarResetFilters"))
 				{
