[ContentBrowser] Isolated filter design updates to the ContentBrowser (not outliner, etc)

diff --git a/Engine/Source/Editor/ContentBrowser/Private/SFilterList.cpp b/Engine/Source/Editor/ContentBrowser/Private/SFilterList.cpp
--- a/Engine/Source/Editor/ContentBrowser/Private/SFilterList.cpp
+++ b/Engine/Source/Editor/ContentBrowser/Private/SFilterList.cpp
@@ -12,6 +12,7 @@
 #include "ContentBrowserItem.h"
 #include "ContentBrowserItemData.h"
 #include "ContentBrowserMenuContexts.h"
+#include "ContentBrowserStyle.h"
 #include "ContentBrowserUtils.h"
 #include "Filters.h"
 #include "Filters/FilterBarConfig.h"
@@ -29,6 +30,7 @@
 #include "Logging/LogMacros.h"
 #include "Math/Vector2D.h"
 #include "Misc/AssertionMacros.h"
+#include "SContentBrowser.h"
 #include "SlateGlobals.h"
 #include "SlotBase.h"
 #include "Styling/SlateTypes.h"
@@ -194,7 +196,15 @@ FARFilter SFilterList::GetCombinedBackendFilter(TArray<TSharedRef<const FPathPer
 	{
 		OutPermissionLists.Add(OtherDevelopersFilter->GetPathPermissionList());
 	}
-	return Super::GetCombinedBackendFilter();
+
+	FARFilter CombinedFilter = Super::GetCombinedBackendFilter();
+
+	if (UE::Editor::ContentBrowser::IsNewStyleEnabled())
+	{
+		CombinedFilter.bRecursivePaths = this->bFilterPathsRecursively;
+	}
+
+	return CombinedFilter;
 }
 
 void SFilterList::DisableFiltersThatHideItems(TArrayView<const FContentBrowserItem> ItemList)
@@ -373,14 +383,206 @@ UAssetFilterBarContext* SFilterList::CreateAssetFilterBarContext()
 {
 	UAssetFilterBarContext* AssetFilterBarContext = SAssetFilterBar<const FContentBrowserItem&>::CreateAssetFilterBarContext();
 
+	// Override PopulateFilterMenu - ContentBrowser has a different menu layout to SAssetFilterBar
+	AssetFilterBarContext->PopulateFilterMenu = FOnPopulateAddAssetFilterMenu::CreateSP(this, &SFilterList::PopulateAddFilterMenu);
+
 	AssetFilterBarContext->OnFilterAssetType = FOnFilterAssetType::CreateLambda([this](UClass *TestClass)
 	{
 		return !IsFilteredByPicker(this->InitialClassFilters, TestClass);
 	});
-	
+
 	return AssetFilterBarContext;
 }
 
+void SFilterList::PopulateAddFilterMenu(UToolMenu* Menu, TSharedPtr<FFilterCategory> MenuExpansion, FOnFilterAssetType OnFilterAssetType)
+{
+	if (!UE::Editor::ContentBrowser::IsNewStyleEnabled())
+	{
+		Super::PopulateAddFilterMenu(Menu, MenuExpansion, OnFilterAssetType);
+		return;
+	}
+
+	using namespace UE::Editor::Widgets;
+
+	TMap<TSharedPtr<FFilterCategory>, FFilterCategoryMenu> CategoryToMenuMap = BuildCategoryToMenuMap(AssetFilterCategories, CustomClassFilters, OnFilterAssetType);
+
+	// Populate the common filter sections (Reset Filters etc)
+	{
+		this->PopulateCommonFilterSections(Menu);
+
+		// Remove legacy section
+		Menu->RemoveSection("FilterBarTextFilters");
+
+		FToolMenuSection& AdvancedSection = Menu->FindOrAddSection(
+			"FilterBarAdvanced",
+			LOCTEXT("FilterBarAdvancedSection", "Advanced"),
+			FToolMenuInsert("BasicFilterBarFiltersMenu", EToolMenuInsertType::Last));
+
+		// Only add the custom text filter submenu if we have a valid CreateTextFilter delegate to use
+		if (this->CreateTextFilter.IsBound())
+		{
+			AdvancedSection.AddSubMenu(
+				"CustomFiltersSubMenu",
+				LOCTEXT("FilterBarTextFilters", "Custom Filters"),
+				LOCTEXT("FilterBarTextFiltersTooltip", "Custom Filters"),
+				FNewToolMenuDelegate::CreateSP(this, &SFilterList::CreateTextFiltersMenu),
+				false,
+				FSlateIcon(FAppStyle::Get().GetStyleSetName(), "Icons.Filter"));
+		}
+
+		if (FToolMenuSection* ManageSection = Menu->FindSection("FilterBarResetFilters"))
+		{
+			ManageSection->Label = LOCTEXT("FilterBarManageSection", "Manage");
+
+			ManageSection->AddMenuEntry(
+				"CopyFilters",
+				LOCTEXT("FilterListCopyFilters", "Copy Filters"),
+				LOCTEXT("FilterListCopyFiltersTooltip", "Copy the current filter selection"),
+				FSlateIcon(FAppStyle::Get().GetStyleSetName(), "GenericCommands.Copy"),
+				FUIAction(
+					FExecuteAction::CreateSP(this, &SFilterList::OnCopyFilters))
+			);
+
+			ManageSection->AddMenuEntry(
+				"PasteFilters",
+				LOCTEXT("FilterListPasteFilters", "Paste Filters"),
+				LOCTEXT("FilterListPasteFiltersTooltip", "Paste to the current filter selection"),
+				FSlateIcon(FAppStyle::Get().GetStyleSetName(), "GenericCommands.Paste"),
+				FUIAction(
+					FExecuteAction::CreateSP(this, &SFilterList::OnPasteFilters),
+					FCanExecuteAction::CreateSP(this, &SFilterList::CanPasteFilters))
+			);
+
+			ManageSection->AddMenuEntry(
+				"FilterRecursively",
+				LOCTEXT("FilterListFilterRecursively", "Filter Recursively"),
+				LOCTEXT("FilterListFilterRecursivelyTooltip", "Apply the current filter selection recursively, relevant to the current path"),
+				{ },
+				FUIAction(
+					FExecuteAction::CreateSP(this, &SFilterList::ToggleFilterPathsRecursively),
+					FCanExecuteAction(),
+					FIsActionChecked::CreateSP(this, &SFilterList::IsFilteringPathsRecursively)),
+				EUserInterfaceActionType::Check
+			);
+		}
+
+		FToolMenuSection& DisplaySection = Menu->AddSection(
+			"FilterBarDisplay",
+			LOCTEXT("FilterListDisplaySection", "Filter Display"));
+		{
+			DisplaySection.AddMenuEntry(
+				"HorizontalLayout",
+				LOCTEXT("FilterListHorizontalLayout", "Horizontal"),
+				LOCTEXT("FilterListHorizontalLayoutToolTip", "Swap to a Horizontal layout for the filter bar"),
+				FSlateIcon(),
+				FUIAction(
+					FExecuteAction::CreateSPLambda(
+					this, [this]
+					{
+						if (this->GetFilterLayout() != EFilterBarLayout::Horizontal)
+						{
+							this->SetFilterLayout(EFilterBarLayout::Horizontal);
+						}
+					}),
+					FCanExecuteAction(),
+					FIsActionChecked::CreateSPLambda(this, [this]
+					{
+						return this->GetFilterLayout() == EFilterBarLayout::Horizontal;
+					})),
+				EUserInterfaceActionType::RadioButton
+			);
+
+			DisplaySection.AddMenuEntry(
+				"VerticalLayout",
+				LOCTEXT("FilterListVerticalLayout", "Vertical"),
+				LOCTEXT("FilterListVerticalLayoutToolTip", "Swap to a vertical layout for the filter bar"),
+				FSlateIcon(),
+				FUIAction(
+					FExecuteAction::CreateSPLambda(this, [this]
+					{
+						if (this->GetFilterLayout() != EFilterBarLayout::Vertical)
+						{
+							this->SetFilterLayout(EFilterBarLayout::Vertical);
+						}
+					}),
+					FCanExecuteAction(),
+					FIsActionChecked::CreateSPLambda(this, [this]
+					{
+						return this->GetFilterLayout() == EFilterBarLayout::Vertical;
+					})),
+				EUserInterfaceActionType::RadioButton
+			);
+		}
+	}
+
+	// If we want to expand a category
+	if (MenuExpansion)
+	{
+		// First add the expanded category, this appears as standard entries in the list (Note: intentionally not using FindChecked here as removing it from the map later would cause the ref to be garbage)
+		if (FFilterCategoryMenu* ExpandedCategory = CategoryToMenuMap.Find( MenuExpansion ))
+		{
+			FToolMenuSection& Section = Menu->AddSection(ExpandedCategory->SectionExtensionHook, ExpandedCategory->SectionHeading);
+
+			// If we are doing a full menu (i.e expanding basic) we add a menu entry which toggles all other categories
+            Section.AddMenuEntry(
+                FName(FText::AsCultureInvariant(ExpandedCategory->SectionHeading).ToString()),
+                ExpandedCategory->SectionHeading,
+                MenuExpansion->Tooltip,
+                FSlateIcon(FAppStyle::Get().GetStyleSetName(), "PlacementBrowser.Icons.Basic"),
+                FUIAction(
+                FExecuteAction::CreateSP( this, &SFilterList::FilterByTypeCategoryClicked, MenuExpansion, ExpandedCategory->Classes ),
+                FCanExecuteAction(),
+                FGetActionCheckState::CreateSP(this, &SFilterList::IsTypeCategoryChecked, MenuExpansion, ExpandedCategory->Classes ) ),
+                EUserInterfaceActionType::ToggleButton
+                );
+
+			Section.AddSeparator("ExpandedCategorySeparator");
+
+			// Now populate with all the assets from the expanded category
+			SFilterList::CreateFiltersMenuCategory( Section, ExpandedCategory->Classes);
+
+			// Remove the Expanded from the map now, as this is treated differently and is no longer needed.
+			CategoryToMenuMap.Remove(MenuExpansion);
+		}
+	}
+
+	// Sort by category name so that we add the submenus in alphabetical order
+	CategoryToMenuMap.KeySort([](const TSharedPtr<FFilterCategory>& A, const TSharedPtr<FFilterCategory>& B) {
+		return A->Title.CompareTo(B->Title) < 0;
+	});
+
+	FToolMenuSection& AdvancedSection = Menu->FindOrAddSection("FilterBarAdvanced");
+
+	FToolMenuEntry& AllFiltersSubMenu = AdvancedSection.AddSubMenu(
+		"AllFilters",
+		LOCTEXT("AllAssetsMenuHeading", "All Filters"),
+		{ },
+		FNewToolMenuDelegate::CreateLambda([this, CategoryToMenuMap](UToolMenu* InMenu){
+			// Add all the other categories as submenus, un-labelled, acts only as a root for the submenus
+			FToolMenuSection& TypeFiltersSection = InMenu->AddSection("AssetFilterBarFilterAdvancedAsset", { });
+
+			for (const TPair<TSharedPtr<FFilterCategory>, FFilterCategoryMenu>& CategoryMenuPair : CategoryToMenuMap)
+			{
+				TypeFiltersSection.AddSubMenu(
+					FName(FText::AsCultureInvariant(CategoryMenuPair.Key->Title).ToString()),
+					CategoryMenuPair.Key->Title,
+					CategoryMenuPair.Key->Tooltip,
+					FNewToolMenuDelegate::CreateSP(this, &SFilterList::CreateFiltersMenuCategory, CategoryMenuPair.Value.Classes),
+					FUIAction(
+					FExecuteAction::CreateSP(this, &SFilterList::FilterByTypeCategoryClicked, CategoryMenuPair.Key, CategoryMenuPair.Value.Classes),
+					FCanExecuteAction(),
+					FGetActionCheckState::CreateSP(this, &SFilterList::IsTypeCategoryChecked, CategoryMenuPair.Key, CategoryMenuPair.Value.Classes)),
+					EUserInterfaceActionType::ToggleButton
+				);
+			}
+
+			// Now add all non-asset filters
+			this->PopulateCustomFilters(InMenu);
+		}));
+
+	AllFiltersSubMenu.InsertPosition = FToolMenuInsert("CustomFiltersSubMenu", EToolMenuInsertType::After);
+}
+
 void SFilterList::CreateCustomFilterDialog(const FText& InText)
 {
 	CreateCustomTextFilterFromSearch(InText);
@@ -427,6 +629,25 @@ void SFilterList::ModifyCustomTextFilterByLabel(const FCustomTextFilterData& New
 	}
 }
 
+bool SFilterList::IsFilteringPathsRecursively() const
+{
+	return bFilterPathsRecursively;
+}
+
+void SFilterList::SetFilterPathsRecursively(const bool bInFilterRecursively)
+{
+	if (bFilterPathsRecursively != bInFilterRecursively)
+	{
+		bFilterPathsRecursively = bInFilterRecursively;
+		this->OnFilterChanged.ExecuteIfBound();
+	}
+}
+
+void SFilterList::ToggleFilterPathsRecursively()
+{
+	SetFilterPathsRecursively(!bFilterPathsRecursively);
+}
+
 void SFilterList::OnDeleteCustomTextFilter(const TSharedPtr<ICustomTextFilter<FAssetFilterType>> InFilter)
 {
 	SAssetFilterBar<FAssetFilterType>::OnDeleteCustomTextFilter(InFilter);
@@ -553,12 +774,26 @@ void SFilterList::UpdateCustomTextFilterIncludes(const bool InIncludeClassName,
 	}
 }
 
-void SFilterList::SaveSettings()
+void SFilterList::SaveSettingsInternal(FFilterBarSettings* InSettingsToSave)
 {
+	check(InSettingsToSave);
+
 	// If this instance doesn't want to use the shared settings, save the settings normally
 	if(!bUseSharedSettings)
 	{
-		SAssetFilterBar<FAssetFilterType>::SaveSettings();
+		// Only save the orientation if we allow dynamic modification and saving
+		InSettingsToSave->bIsLayoutSaved = this->bCanChangeOrientation;
+		if (this->bCanChangeOrientation)
+		{
+			InSettingsToSave->FilterBarLayout = this->FilterBarLayout;
+
+			if (UE::Editor::ContentBrowser::IsNewStyleEnabled())
+			{
+				InSettingsToSave->bFilterRecursively = this->bFilterPathsRecursively;
+			}
+		}
+
+		SAssetFilterBar<FAssetFilterType>::SaveSettingsInternal(InSettingsToSave);
 		return;
 	}
 
@@ -670,7 +905,26 @@ void SFilterList::LoadSettings(const FName& InInstanceName)
 	// If this instance doesn't want to use the shared settings, load the settings normally
 	if(!bUseSharedSettings)
 	{
+		const FFilterBarSettings* FilterBarConfig = GetConstConfig();
+
+		if(!FilterBarConfig)
+		{
+			UE_LOG(LogSlate, Error, TEXT("SFilterBar Requires that you specify a FilterBarIdentifier to load settings"));
+			return;
+		}
+
+		/* Only load the setting if we saved it, aka if SaveSettings() was ever called and bCanChangeOrientation is true */
+		if(FilterBarConfig->bIsLayoutSaved)
+		{
+			this->SetFilterLayout(FilterBarConfig->FilterBarLayout);
+			if (UE::Editor::ContentBrowser::IsNewStyleEnabled())
+			{
+				this->SetFilterPathsRecursively(FilterBarConfig->bFilterRecursively);
+			}
+		}
+
 		SAssetFilterBar<FAssetFilterType>::LoadSettings();
+
 		return;
 	}
 
diff --git a/Engine/Source/Editor/ContentBrowser/Private/SFilterList.h b/Engine/Source/Editor/ContentBrowser/Private/SFilterList.h
--- a/Engine/Source/Editor/ContentBrowser/Private/SFilterList.h
+++ b/Engine/Source/Editor/ContentBrowser/Private/SFilterList.h
@@ -146,7 +146,7 @@ public:
 	/** Add a custom widget to the filter bar alongside the filters */
 	void AddWidgetToCurrentLayout(TSharedRef<SWidget> InWidget);
 
-	virtual void SaveSettings() override;
+	using Super::SaveSettings;
 	virtual void LoadSettings() override;
 
 	/* Copy the settings from a specific instance name */
@@ -168,10 +168,21 @@ public:
 	/** Modifies the custom text filter with the matching label */
 	void ModifyCustomTextFilterByLabel(const FCustomTextFilterData& NewFilterData, const FText& FilterLabel);
 
+	/** Whether the filters are being applied to paths recursively */
+	bool IsFilteringPathsRecursively() const;
+
+	/** Sets whether the filters are being applied to paths recursively */
+	void SetFilterPathsRecursively(const bool bInFilterRecursively);
+
+	/** Toggle whether the filters are being applied to paths recursively */
+	void ToggleFilterPathsRecursively();
+
 protected:
 	
 	/** Handler for when the add filter button was clicked */
 	//virtual TSharedRef<SWidget> MakeAddFilterMenu() override;
+
+	virtual void SaveSettingsInternal(FFilterBarSettings* InSettingsToSave) override;
 	
 	/** Handler for when a custom text filter is modified */
 	virtual void OnModifyCustomTextFilter(const FCustomTextFilterData& InFilterData, TSharedPtr<ICustomTextFilter<FAssetFilterType>> InFilter) override;
@@ -185,6 +196,9 @@ private:
 	// TSharedRef<SWidget> MakeAddFilterMenu(FAssetCategoryPath MenuExpansion = EAssetCategoryPaths::Basic);
 
 	virtual UAssetFilterBarContext* CreateAssetFilterBarContext() override;
+
+	/** Handler to Populate the Add Filter Menu. Use OnFilterAssetType in Subclasses to add classes to the exclusion list */
+	virtual void PopulateAddFilterMenu(UToolMenu* Menu, TSharedPtr<FFilterCategory> MenuExpansion, FOnFilterAssetType OnFilterAssetType) override;
 	
 	/** Handler for when another SFilterList using shared settings creates a custom text filter */
 	void OnExternalCustomTextFilterCreated(TSharedPtr<SWidget> BroadcastingFilterList);
@@ -226,6 +240,9 @@ private:
 	/** Whether this FilterList wants to load/save from the settings common to all instances */
 	bool bUseSharedSettings;
 
+	/** Whether this FilterList applies to paths recursively */
+	bool bFilterPathsRecursively = true;
+
 	/** The event that executes when a custom text filter is created/modified/deleted in any filter list that is using Shared Settings */
 	static FCustomTextFilterEvent CustomTextFilterEvent;
 };
diff --git a/Engine/Source/Editor/EditorWidgets/Private/Filters/SAssetFilterBar.cpp b/Engine/Source/Editor/EditorWidgets/Private/Filters/SAssetFilterBar.cpp
--- a/Engine/Source/Editor/EditorWidgets/Private/Filters/SAssetFilterBar.cpp
+++ b/Engine/Source/Editor/EditorWidgets/Private/Filters/SAssetFilterBar.cpp
@@ -2,6 +2,84 @@
 
 #include "Filters/SAssetFilterBar.h"
 
+#define LOCTEXT_NAMESPACE "FilterBar"
+
+namespace UE::Editor::Widgets
+{
+	FFilterCategoryMenu::FFilterCategoryMenu(
+		const FName& InSectionExtensionHook,
+		const FText& InSectionHeading)
+		: SectionExtensionHook(InSectionExtensionHook)
+		, SectionHeading(InSectionHeading)
+	{}
+
+	TMap<TSharedPtr<FFilterCategory>, FFilterCategoryMenu> BuildCategoryToMenuMap(
+		TMap<FName, TSharedPtr<FFilterCategory>>& InAssetFilterCategories,
+		TConstArrayView<TSharedRef<FCustomClassFilterData>> InCustomClassFilters, const FOnFilterAssetType& InOnFilterAssetType)
+	{
+		// Create a map of Categories to Menus
+		TMap<TSharedPtr<FFilterCategory>, FFilterCategoryMenu> CategoryToMenuMap;
+
+		// For every asset type, move it into all the categories it should appear in
+		for (const TSharedRef<FCustomClassFilterData>& CustomClassFilter : InCustomClassFilters)
+		{
+			bool bPassesExternalFilters = true;
+
+			// Run any external class filters we have
+			if (InOnFilterAssetType.IsBound())
+			{
+				bPassesExternalFilters = InOnFilterAssetType.Execute(CustomClassFilter->GetClass());
+			}
+
+			if (bPassesExternalFilters)
+			{
+				/** Get all the categories this filter belongs to */
+				TArray<TSharedPtr<FFilterCategory>> Categories = CustomClassFilter->GetCategories();
+
+				for (const TSharedPtr<FFilterCategory>& Category : Categories)
+				{
+					// If the category for this custom class already exists
+					if (FFilterCategoryMenu* CategoryMenu = CategoryToMenuMap.Find(Category))
+					{
+						CategoryMenu->Classes.Add( CustomClassFilter );
+					}
+					// Otherwise create a new FCategoryMenu for the category and add it to the map
+					else
+					{
+						const FText SectionHeading = FText::Format(LOCTEXT("WildcardFilterHeadingHeadingTooltip", "{0} Filters"), Category->Title);
+						const FName ExtensionPoint = FName(FText::AsCultureInvariant(SectionHeading).ToString());
+
+						FFilterCategoryMenu NewCategoryMenu(ExtensionPoint, SectionHeading);
+						NewCategoryMenu.Classes.Add(CustomClassFilter);
+
+						CategoryToMenuMap.Add(Category, NewCategoryMenu);
+					}
+				}
+			}
+		}
+
+		// Remove any empty categories
+		for (auto MenuIt = CategoryToMenuMap.CreateIterator(); MenuIt; ++MenuIt)
+		{
+			if (MenuIt.Value().Classes.Num() == 0)
+			{
+				CategoryToMenuMap.Remove(MenuIt.Key());
+			}
+		}
+
+		// Set the extension hook for the basic category, if it exists and we have any assets for it
+		if (TSharedPtr<FFilterCategory>* BasicCategory = InAssetFilterCategories.Find(EAssetCategoryPaths::Basic.GetCategory()))
+		{
+			if (FFilterCategoryMenu* BasicMenu = CategoryToMenuMap.Find(*BasicCategory))
+			{
+				BasicMenu->SectionExtensionHook = "FilterBarFilterBasicAsset";
+			}
+		}
+
+		return CategoryToMenuMap;
+	}
+}
+
 FFilterBarSettings* FFilterBarBase::GetMutableConfig()
 {
 	if(FilterBarIdentifier.IsNone())
@@ -38,3 +116,5 @@ void FFilterBarBase::InitializeConfig()
 	// Call GetMutableConfig to force create a config for this filter bar if the user specified FilterBarIdentifier
 	FFilterBarSettings* FilterBarConfig = GetMutableConfig();
 }
+
+#undef LOCTEXT_NAMESPACE
diff --git a/Engine/Source/Editor/EditorWidgets/Public/Filters/SAssetFilterBar.h b/Engine/Source/Editor/EditorWidgets/Public/Filters/SAssetFilterBar.h
--- a/Engine/Source/Editor/EditorWidgets/Public/Filters/SAssetFilterBar.h
+++ b/Engine/Source/Editor/EditorWidgets/Public/Filters/SAssetFilterBar.h
@@ -62,26 +62,6 @@
 class SFilterSearchBox;
 class SWidget;
 
-namespace UE::Editor::EditorWidgets
-{
-	static bool IsNewStyleEnabled()
-	{
-		// @note: This is tied to the ContentBrowser new style (5.6+)
-
-		static bool bIsNewStyleEnabled = false;
-		UE_CALL_ONCE([&]()
-		{
-			if (const IConsoleVariable* CVar = IConsoleManager::Get().FindConsoleVariable(TEXT("ContentBrowser.EnableNewStyle")))
-			{
-				ensureAlwaysMsgf(!EnumHasAnyFlags(CVar->GetFlags(), ECVF_Default), TEXT("The CVar should have already been set from commandline, @see: UnrealEdGlobals.cpp, UE::Editor::ContentBrowser::EnableContentBrowserNewStyleCVarRegistration."));
-				bIsNewStyleEnabled = CVar->GetBool();
-			}
-		});
-
-		return bIsNewStyleEnabled;
-	}
-}
-
 #define LOCTEXT_NAMESPACE "FilterBar"
 
 /** Delegate that subclasses can use to specify classes to not include in this filter
@@ -92,6 +72,27 @@ DECLARE_DELEGATE_RetVal_OneParam(bool, FOnFilterAssetType, UClass* /* Test Class
 /* Delegate used by SAssetFilterBar to populate the add filter menu */
 DECLARE_DELEGATE_ThreeParams(FOnPopulateAddAssetFilterMenu, UToolMenu*, TSharedPtr<FFilterCategory>, FOnFilterAssetType)
 
+namespace UE::Editor::Widgets
+{
+	/** A local struct to describe a category in the filter menu */
+	struct EDITORWIDGETS_API FFilterCategoryMenu
+	{
+		/** The Classes that belong to this category */
+		TArray<TSharedPtr<FCustomClassFilterData>> Classes;
+
+		/** Menu section */
+		FName SectionExtensionHook;
+		FText SectionHeading;
+
+		FFilterCategoryMenu(const FName& InSectionExtensionHook, const FText& InSectionHeading);
+	};
+
+	EDITORWIDGETS_API TMap<TSharedPtr<FFilterCategory>, FFilterCategoryMenu> BuildCategoryToMenuMap(
+		TMap<FName, TSharedPtr<FFilterCategory>>& InAssetFilterCategories,
+		TConstArrayView<TSharedRef<FCustomClassFilterData>> InCustomClassFilters,
+		const FOnFilterAssetType& InOnFilterAssetType);
+}
+
 /** ToolMenuContext that is used to create the Add Filter Menu */
 UCLASS()
 class EDITORWIDGETS_API UAssetFilterBarContext : public UObject
@@ -274,7 +275,7 @@ public:
 	virtual void SaveSettings()
 	{
 		FFilterBarSettings* FilterBarConfig = GetMutableConfig();
-		
+
 		if(!FilterBarConfig)
 		{
 			UE_LOG(LogSlate, Error, TEXT("SFilterBar Requires that you specify a FilterBarIdentifier to save settings"));
@@ -284,22 +285,7 @@ public:
 		// Empty the config, we are just going to re-save everything
 		FilterBarConfig->Empty();
 
-		SaveCustomTextFilters(FilterBarConfig);
-		SaveFilters(FilterBarConfig);
-
-		// Only save the orientation if we allow dynamic modification and saving
-		FilterBarConfig->bIsLayoutSaved = this->bCanChangeOrientation;
-		if(this->bCanChangeOrientation)
-		{
-			FilterBarConfig->FilterBarLayout = this->FilterBarLayout;
-
-			if (UE::Editor::EditorWidgets::IsNewStyleEnabled())
-			{
-				FilterBarConfig->bFilterRecursively = this->bFilterPathsRecursively;
-			}
-		}
-
-		SaveConfig();
+		SaveSettingsInternal(FilterBarConfig);
 	}
 
 	virtual void LoadSettings()
@@ -321,16 +307,6 @@ public:
 		LoadFilters(FilterBarConfig);
 		LoadCustomTextFilters(FilterBarConfig);
 
-		/* Only load the setting if we saved it, aka if SaveSettings() was ever called and bCanChangeOrientation is true */
-		if(FilterBarConfig->bIsLayoutSaved)
-		{
-			this->SetFilterLayout(FilterBarConfig->FilterBarLayout);
-			if (UE::Editor::EditorWidgets::IsNewStyleEnabled())
-			{
-				this->SetFilterPathsRecursively(FilterBarConfig->bFilterRecursively);
-			}
-		}
-
 		this->OnFilterChanged.ExecuteIfBound();
 	}
 
@@ -386,6 +362,16 @@ protected:
 		}
 	}
 
+	virtual void SaveSettingsInternal(FFilterBarSettings* InSettingsToSave)
+	{
+		check(InSettingsToSave);
+
+		SaveCustomTextFilters(InSettingsToSave);
+		SaveFilters(InSettingsToSave);
+
+		SaveConfig();
+	}
+
 	/** Save all the custom filters (created programatically) into the specified config */
 	void SaveFilters(FFilterBarSettings* FilterBarConfig)
 	{
@@ -670,11 +656,6 @@ public:
 			}
 		}
 
-		if (UE::Editor::EditorWidgets::IsNewStyleEnabled())
-		{
-			CombinedFilter.bRecursivePaths = this->bFilterPathsRecursively;
-		}
-
 		return CombinedFilter;
 	}
 
@@ -1030,279 +1011,60 @@ protected:
 		FToolMenuContext ToolMenuContext(CreateAssetFilterBarContext());
 		return UToolMenus::Get()->GenerateWidget(FilterMenuName, ToolMenuContext);
 	} 
-	
+
 	/** Handler to Populate the Add Filter Menu. Use OnFilterAssetType in Subclasses to add classes to the exclusion list */
-	void PopulateAddFilterMenu(UToolMenu* Menu, TSharedPtr<FFilterCategory> MenuExpansion, FOnFilterAssetType OnFilterAssetType)
+	virtual void PopulateAddFilterMenu(UToolMenu* Menu, TSharedPtr<FFilterCategory> MenuExpansion, FOnFilterAssetType OnFilterAssetType)
 	{
-		FAssetToolsModule& AssetToolsModule = FModuleManager::LoadModuleChecked<FAssetToolsModule>(TEXT("AssetTools"));
-
-		// A local struct to describe a category in the filter menu
-		struct FCategoryMenu
-		{
-			/** The Classes that belong to this category */
-			TArray<TSharedPtr<FCustomClassFilterData>> Classes;
-
-			//Menu section
-			FName SectionExtensionHook;
-			FText SectionHeading;
-
-			FCategoryMenu(const FName& InSectionExtensionHook, const FText& InSectionHeading)
-				: SectionExtensionHook(InSectionExtensionHook)
-				, SectionHeading(InSectionHeading)
-			{}
-		};
-		
-		// Create a map of Categories to Menus
-		TMap<TSharedPtr<FFilterCategory>, FCategoryMenu> CategoryToMenuMap;
-		
-		// For every asset type, move it into all the categories it should appear in
-		for(const TSharedRef<FCustomClassFilterData> &CustomClassFilter : CustomClassFilters)
-		{
-			bool bPassesExternalFilters = true;
-
-			// Run any external class filters we have
-			if(OnFilterAssetType.IsBound())
-			{
-				bPassesExternalFilters = OnFilterAssetType.Execute(CustomClassFilter->GetClass());
-			}
-
-			if(bPassesExternalFilters)
-			{
-				/** Get all the categories this filter belongs to */
-				TArray<TSharedPtr<FFilterCategory>> Categories = CustomClassFilter->GetCategories();
-
-				for(const TSharedPtr<FFilterCategory>& Category : Categories)
-				{
-					// If the category for this custom class already exists
-					if(FCategoryMenu* CategoryMenu = CategoryToMenuMap.Find(Category))
-					{
-						CategoryMenu->Classes.Add( CustomClassFilter );
-					}
-					// Otherwise create a new FCategoryMenu for the category and add it to the map
-					else
-					{
-						const FText SectionHeading = FText::Format(LOCTEXT("WildcardFilterHeadingHeadingTooltip", "{0} Filters"), Category->Title);
-						const FName ExtensionPoint = FName(FText::AsCultureInvariant(SectionHeading).ToString());
-
-						FCategoryMenu NewCategoryMenu(ExtensionPoint, SectionHeading);
-						NewCategoryMenu.Classes.Add(CustomClassFilter);
+		using namespace UE::Editor::Widgets;
 
-						CategoryToMenuMap.Add(Category, NewCategoryMenu);
-					}
-				}
-			}
-		}
-
-		// Remove any empty categories
-		for (auto MenuIt = CategoryToMenuMap.CreateIterator(); MenuIt; ++MenuIt)
-		{
-			if (MenuIt.Value().Classes.Num() == 0)
-			{
-				CategoryToMenuMap.Remove(MenuIt.Key());
-			}
-		}
-
-		// Set the extension hook for the basic category, if it exists and we have any assets for it
-		if(TSharedPtr<FFilterCategory>* BasicCategory = AssetFilterCategories.Find(EAssetCategoryPaths::Basic.GetCategory()))
-		{
-			if(FCategoryMenu* BasicMenu = CategoryToMenuMap.Find(*BasicCategory))
-			{
-				BasicMenu->SectionExtensionHook = "FilterBarFilterBasicAsset";
-			}
-		}
+		TMap<TSharedPtr<FFilterCategory>, FFilterCategoryMenu> CategoryToMenuMap = BuildCategoryToMenuMap(AssetFilterCategories, CustomClassFilters, OnFilterAssetType);
 
 		// Populate the common filter sections (Reset Filters etc)
 		{
 			this->PopulateCommonFilterSections(Menu);
-
-			if (UE::Editor::EditorWidgets::IsNewStyleEnabled())
-			{
-				// Remove legacy section
-				Menu->RemoveSection("FilterBarTextFilters");
-
-				FToolMenuSection& AdvancedSection = Menu->FindOrAddSection(
-					"FilterBarAdvanced",
-					LOCTEXT("FilterBarAdvancedSection", "Advanced"),
-					FToolMenuInsert("BasicFilterBarFiltersMenu", EToolMenuInsertType::Last));
-
-				// Only add the custom text filter submenu if we have a valid CreateTextFilter delegate to use
-				if (this->CreateTextFilter.IsBound())
-				{
-					AdvancedSection.AddSubMenu(
-						"CustomFiltersSubMenu",
-						LOCTEXT("FilterBarTextFilters", "Custom Filters"),
-						LOCTEXT("FilterBarTextFiltersTooltip", "Custom Filters"),
-						FNewToolMenuDelegate::CreateSP(this, &SAssetFilterBar<FilterType>::CreateTextFiltersMenu),
-						false,
-						FSlateIcon(FAppStyle::Get().GetStyleSetName(), "Icons.Filter"));
-				}
-
-				if (FToolMenuSection* ManageSection = Menu->FindSection("FilterBarResetFilters"))
-				{
-					ManageSection->Label = LOCTEXT("FilterBarManageSection", "Manage");
-
-					ManageSection->AddMenuEntry(
-						"CopyFilters",
-						LOCTEXT("FilterListCopyFilters", "Copy Filters"),
-						LOCTEXT("FilterListCopyFiltersTooltip", "Copy the current filter selection"),
-						FSlateIcon(FAppStyle::Get().GetStyleSetName(), "GenericCommands.Copy"),
-						FUIAction(
-							FExecuteAction::CreateSP(this, &SAssetFilterBar<FilterType>::OnCopyFilters))
-					);
-
-					ManageSection->AddMenuEntry(
-						"PasteFilters",
-						LOCTEXT("FilterListPasteFilters", "Paste Filters"),
-						LOCTEXT("FilterListPasteFiltersTooltip", "Paste to the current filter selection"),
-						FSlateIcon(FAppStyle::Get().GetStyleSetName(), "GenericCommands.Paste"),
-						FUIAction(
-							FExecuteAction::CreateSP(this, &SAssetFilterBar<FilterType>::OnPasteFilters),
-							FCanExecuteAction::CreateSP(this, &SAssetFilterBar<FilterType>::CanPasteFilters))
-					);
-
-					ManageSection->AddMenuEntry(
-						"FilterRecursively",
-						LOCTEXT("FilterListFilterRecursively", "Filter Recursively"),
-						LOCTEXT("FilterListFilterRecursivelyTooltip", "Apply the current filter selection recursively, relevant to the current path"),
-						{ },
-						FUIAction(
-							FExecuteAction::CreateSP(this, &SAssetFilterBar<FilterType>::OnFilterPathsRecursively),
-							FCanExecuteAction(),
-							FIsActionChecked::CreateSP(this, &SAssetFilterBar<FilterType>::IsFilteringPathsRecursively)),
-						EUserInterfaceActionType::Check
-					);
-				}
-
-				FToolMenuSection& DisplaySection = Menu->AddSection(
-					"FilterBarDisplay",
-					LOCTEXT("FilterListDisplaySection", "Filter Display"));
-				{
-					DisplaySection.AddMenuEntry(
-						"HorizontalLayout",
-						LOCTEXT("FilterListHorizontalLayout", "Horizontal"),
-						LOCTEXT("FilterListHorizontalLayoutToolTip", "Swap to a Horizontal layout for the filter bar"),
-						FSlateIcon(),
-						FUIAction(
-							FExecuteAction::CreateSPLambda(
-							this, [this]
-							{
-								if (this->GetFilterLayout() != EFilterBarLayout::Horizontal)
-								{
-									this->SetFilterLayout(EFilterBarLayout::Horizontal);
-								}
-							}),
-							FCanExecuteAction(),
-							FIsActionChecked::CreateSPLambda(this, [this]
-							{
-								return this->GetFilterLayout() == EFilterBarLayout::Horizontal;
-							})),
-						EUserInterfaceActionType::RadioButton
-					);
-
-					DisplaySection.AddMenuEntry(
-						"VerticalLayout",
-						LOCTEXT("FilterListVerticalLayout", "Vertical"),
-						LOCTEXT("FilterListVerticalLayoutToolTip", "Swap to a vertical layout for the filter bar"),
-						FSlateIcon(),
-						FUIAction(
-							FExecuteAction::CreateSPLambda(this, [this]
-							{
-								if (this->GetFilterLayout() != EFilterBarLayout::Vertical)
-								{
-									this->SetFilterLayout(EFilterBarLayout::Vertical);
-								}
-							}),
-							FCanExecuteAction(),
-							FIsActionChecked::CreateSPLambda(this, [this]
-							{
-								return this->GetFilterLayout() == EFilterBarLayout::Vertical;
-							})),
-						EUserInterfaceActionType::RadioButton
-					);
-				}
-			}
 		}
-		
+
 		// If we want to expand a category
-		if(MenuExpansion)
+		if (MenuExpansion)
 		{
 			// First add the expanded category, this appears as standard entries in the list (Note: intentionally not using FindChecked here as removing it from the map later would cause the ref to be garbage)
-			FCategoryMenu* ExpandedCategory = CategoryToMenuMap.Find( MenuExpansion );
-			if(ExpandedCategory)
+			if (FFilterCategoryMenu* ExpandedCategory = CategoryToMenuMap.Find( MenuExpansion ))
 			{
 				FToolMenuSection& Section = Menu->AddSection(ExpandedCategory->SectionExtensionHook, ExpandedCategory->SectionHeading);
 
-				if (!UE::Editor::EditorWidgets::IsNewStyleEnabled())
-				{
-					// If we are doing a full menu (i.e expanding basic) we add a menu entry which toggles all other categories
-					Section.AddMenuEntry(
-						FName(FText::AsCultureInvariant(ExpandedCategory->SectionHeading).ToString()),
-						ExpandedCategory->SectionHeading,
-						MenuExpansion->Tooltip,
-						FSlateIcon(FAppStyle::Get().GetStyleSetName(), "PlacementBrowser.Icons.Basic"),
-						FUIAction(
-						FExecuteAction::CreateSP( this, &SAssetFilterBar<FilterType>::FilterByTypeCategoryClicked, MenuExpansion, ExpandedCategory->Classes ),
-						FCanExecuteAction(),
-						FGetActionCheckState::CreateSP(this, &SAssetFilterBar<FilterType>::IsTypeCategoryChecked, MenuExpansion, ExpandedCategory->Classes ) ),
-						EUserInterfaceActionType::ToggleButton
-						);
-
-					Section.AddSeparator("ExpandedCategorySeparator");
-				}
+				// If we are doing a full menu (i.e expanding basic) we add a menu entry which toggles all other categories
+                Section.AddMenuEntry(
+                	FName(FText::AsCultureInvariant(ExpandedCategory->SectionHeading).ToString()),
+                	ExpandedCategory->SectionHeading,
+                	MenuExpansion->Tooltip,
+                	FSlateIcon(FAppStyle::Get().GetStyleSetName(), "PlacementBrowser.Icons.Basic"),
+                	FUIAction(
+                	FExecuteAction::CreateSP( this, &SAssetFilterBar<FilterType>::FilterByTypeCategoryClicked, MenuExpansion, ExpandedCategory->Classes ),
+                	FCanExecuteAction(),
+                	FGetActionCheckState::CreateSP(this, &SAssetFilterBar<FilterType>::IsTypeCategoryChecked, MenuExpansion, ExpandedCategory->Classes ) ),
+                	EUserInterfaceActionType::ToggleButton
+                	);
+
+				Section.AddSeparator("ExpandedCategorySeparator");
 
 				// Now populate with all the assets from the expanded category
-				SAssetFilterBar<FilterType>::CreateFiltersMenuCategory(Section, ExpandedCategory->Classes);
+				SAssetFilterBar<FilterType>::CreateFiltersMenuCategory( Section, ExpandedCategory->Classes);
 
 				// Remove the Expanded from the map now, as this is treated differently and is no longer needed.
 				CategoryToMenuMap.Remove(MenuExpansion);
 			}
 		}
 
+		// Add all the other categories as submenus
+		FToolMenuSection& Section = Menu->AddSection("AssetFilterBarFilterAdvancedAsset", LOCTEXT("AdvancedAssetsMenuHeading", "Type Filters"));
+
 		// Sort by category name so that we add the submenus in alphabetical order
 		CategoryToMenuMap.KeySort([](const TSharedPtr<FFilterCategory>& A, const TSharedPtr<FFilterCategory>& B) {
 			return A->Title.CompareTo(B->Title) < 0;
 		});
 
-if (UE::Editor::EditorWidgets::IsNewStyleEnabled())
-{
-		FToolMenuSection& AdvancedSection = Menu->FindOrAddSection("FilterBarAdvanced");
-
-		FToolMenuEntry& AllFiltersSubMenu = AdvancedSection.AddSubMenu(
-			"AllFilters",
-			LOCTEXT("AllAssetsMenuHeading", "All Filters"),
-			{ },
-			FNewToolMenuDelegate::CreateLambda([this, CategoryToMenuMap](UToolMenu* InMenu){
-				// Add all the other categories as submenus, un-labelled, acts only as a root for the submenus
-				FToolMenuSection& TypeFiltersSection = InMenu->AddSection("AssetFilterBarFilterAdvancedAsset", { });
-
-				for (const TPair<TSharedPtr<FFilterCategory>, FCategoryMenu>& CategoryMenuPair : CategoryToMenuMap)
-				{
-					TypeFiltersSection.AddSubMenu(
-						FName(FText::AsCultureInvariant(CategoryMenuPair.Key->Title).ToString()),
-						CategoryMenuPair.Key->Title,
-						CategoryMenuPair.Key->Tooltip,
-						FNewToolMenuDelegate::CreateSP(this, &SAssetFilterBar<FilterType>::CreateFiltersMenuCategory, CategoryMenuPair.Value.Classes),
-						FUIAction(
-						FExecuteAction::CreateSP(this, &SAssetFilterBar<FilterType>::FilterByTypeCategoryClicked, CategoryMenuPair.Key, CategoryMenuPair.Value.Classes),
-						FCanExecuteAction(),
-						FGetActionCheckState::CreateSP(this, &SAssetFilterBar<FilterType>::IsTypeCategoryChecked, CategoryMenuPair.Key, CategoryMenuPair.Value.Classes)),
-						EUserInterfaceActionType::ToggleButton
-					);
-				}
-
-				// Now add all non-asset filters
-				this->PopulateCustomFilters(InMenu);
-			}));
-
-		AllFiltersSubMenu.InsertPosition = FToolMenuInsert("CustomFiltersSubMenu", EToolMenuInsertType::After);
-}
-else
-{
-		// Add all the other categories as submenus
-		FToolMenuSection& Section = Menu->AddSection("AssetFilterBarFilterAdvancedAsset", LOCTEXT("AdvancedAssetsMenuHeading", "Type Filters"));
-
 		// Now actually add them
-		for (const TPair<TSharedPtr<FFilterCategory>, FCategoryMenu>& CategoryMenuPair : CategoryToMenuMap)
+		for (const TPair<TSharedPtr<FFilterCategory>, FFilterCategoryMenu>& CategoryMenuPair : CategoryToMenuMap)
 		{
 			Section.AddSubMenu(
 				FName(FText::AsCultureInvariant(CategoryMenuPair.Key->Title).ToString()),
@@ -1319,7 +1081,6 @@ else
 
 		// Now add all non-asset filters
 		this->PopulateCustomFilters(Menu);
-}
 	}
 
 	/* Asset Type filter related functionality */
@@ -1467,25 +1228,6 @@ else
 		return FFilterBarSettings::StaticStruct()->ImportText(*Output, &Value, nullptr, 0, GLog, GetPathNameSafe(FFilterBarSettings::StaticStruct())) != nullptr;
 	}
 
-	void SetFilterPathsRecursively(const bool bInFilterRecursively)
-	{
-		if (bFilterPathsRecursively != bInFilterRecursively)
-		{
-			bFilterPathsRecursively = bInFilterRecursively;
-			this->OnFilterChanged.ExecuteIfBound();
-		}
-	}
-
-	void OnFilterPathsRecursively()
-	{
-		SetFilterPathsRecursively(!bFilterPathsRecursively);
-	}
-
-	bool IsFilteringPathsRecursively() const
-	{
-		return bFilterPathsRecursively;
-	}
-
 protected:
 	/** A copy of all AssetFilters in this->Filters for convenient access */
 	TArray< TSharedRef<SAssetFilter> > AssetFilters;
@@ -1510,8 +1252,6 @@ protected:
 	 *  bar. These are stored so we can save these when calling SaveSettings() so that they are not lost */
 	TMap<FString, bool> UnknownTypeFilters;
 
-	bool bFilterPathsRecursively = true;
-
 private:
 	/** The filter menu category to expand. */
 	TOptional<FAssetCategoryPath> DefaultMenuExpansionCategory;
