Sort subcategories in filters

#tests editor

diff --git a/Engine/Source/Editor/EditorWidgets/Public/Filters/SAssetFilterBar.h b/Engine/Source/Editor/EditorWidgets/Public/Filters/SAssetFilterBar.h
--- a/Engine/Source/Editor/EditorWidgets/Public/Filters/SAssetFilterBar.h
+++ b/Engine/Source/Editor/EditorWidgets/Public/Filters/SAssetFilterBar.h
@@ -1063,6 +1063,15 @@ protected:
 			return A->Title.CompareTo(B->Title) < 0;
 		});
 
+		// Sort each submenu internals 
+		for (TPair<TSharedPtr<FFilterCategory>, FFilterCategoryMenu>& CategoryToMenu : CategoryToMenuMap)
+		{
+			CategoryToMenu.Value.Classes.Sort([](const TSharedPtr<FCustomClassFilterData>& A, const TSharedPtr<FCustomClassFilterData>& B)
+			{
+				return A->GetName().CompareTo(B->GetName()) < 0;
+			});
+		}
+
 		// Now actually add them
 		for (const TPair<TSharedPtr<FFilterCategory>, FFilterCategoryMenu>& CategoryMenuPair : CategoryToMenuMap)
 		{
