Editor Widgets: Fixed missing argument being passed to MakeTransportControlWidget

#rnx

diff --git a/Engine/Source/Editor/EditorWidgets/Private/STransportControl.cpp b/Engine/Source/Editor/EditorWidgets/Private/STransportControl.cpp
--- a/Engine/Source/Editor/EditorWidgets/Private/STransportControl.cpp
+++ b/Engine/Source/Editor/EditorWidgets/Private/STransportControl.cpp
@@ -182,7 +182,7 @@ void STransportControl::Construct( const STransportControl::FArguments& InArgs )
 	{
 		for(FTransportControlWidget WidgetDesc : TransportControlArgs.WidgetsToCreate)
 		{
-			TSharedPtr<SWidget> Widget = MakeTransportControlWidget(WidgetDesc.WidgetType);
+			TSharedPtr<SWidget> Widget = MakeTransportControlWidget(WidgetDesc.WidgetType, WidgetDesc.MakeCustomWidgetDelegate);
 			if(Widget.IsValid())
 			{
 				HorizontalBox->AddSlot()
