Created re-usable input chord widget

diff --git a/Engine/Source/Editor/EditorWidgets/Private/SInputChord.cpp b/Engine/Source/Editor/EditorWidgets/Private/SInputChord.cpp
new file mode 100644
--- /dev/null
+++ b/Engine/Source/Editor/EditorWidgets/Private/SInputChord.cpp
@@ -0,0 +1,117 @@
+// Copyright Epic Games, Inc. All Rights Reserved.
+
+#include "SInputChord.h"
+
+#include "Components/HorizontalBox.h"
+#include "SlateOptMacros.h"
+#include "Styling/AppStyle.h"
+#include "Widgets/Images/SImage.h"
+#include "Widgets/Layout/SSeparator.h"
+#include "Widgets/Text/STextBlock.h"
+
+BEGIN_SLATE_FUNCTION_BUILD_OPTIMIZATION
+
+void SInputChord::Construct(const FArguments& InArgs)
+{
+	// Padding between Icon and InputLabel, and between either side of the InputLabel delimiter
+	static constexpr float HorizontalElementPadding = 5.0f;
+	static constexpr float DefaultVerticalPadding = 2.0f;
+
+	TSharedRef<SHorizontalBox> ContainerWidget = SNew(SHorizontalBox);
+
+	// Prepend Icon, if set
+	if (InArgs._Icon->IsSet())
+	{
+		ContainerWidget->AddSlot()
+		.AutoWidth()
+		.HAlign(HAlign_Right)
+		.VAlign(VAlign_Center)
+		.Padding(FMargin{0.0f, DefaultVerticalPadding, HorizontalElementPadding, DefaultVerticalPadding})
+		[
+			SNew(SImage)
+			.Image(InArgs._Icon)
+		];
+	}
+
+	// Always display the input label (key combo)
+	ContainerWidget->AddSlot()
+	.AutoWidth()
+	.HAlign(HAlign_Right)
+	.VAlign(VAlign_Center)
+	[
+		SNew(STextBlock)
+		.TextStyle(InArgs._InputLabelStyle)
+		.Text(InArgs._InputLabel)
+	];
+
+	// Only add the action label and delimiter if the action label is set
+	if (!InArgs._ActionLabel.IsEmpty())
+	{
+		ContainerWidget->AddSlot()
+		.AutoWidth()
+		.HAlign(HAlign_Center)
+		.VAlign(VAlign_Center)
+		.Padding(FMargin{HorizontalElementPadding, DefaultVerticalPadding})
+		[
+			InArgs._InputLabelDelimiterOverride.IsValid()
+			? InArgs._InputLabelDelimiterOverride.ToSharedRef()
+			: SNew(SSeparator)
+			.ColorAndOpacity(FLinearColor{1,1,1, 0.4f})
+			.SeparatorImage(FAppStyle::Get().GetBrush("WhiteBrush"))
+			.Orientation(EOrientation::Orient_Vertical)
+			.Thickness(1.0f)
+		];
+
+		ContainerWidget->AddSlot()
+		.AutoWidth()
+		.HAlign(HAlign_Left)
+		.VAlign(VAlign_Center)
+		[
+			SNew(STextBlock)
+			.TextStyle(InArgs._ActionLabelStyle)
+			.Text(InArgs._ActionLabel)
+		];
+	}
+
+	ChildSlot
+	.HAlign(HAlign_Fill)
+	.VAlign(VAlign_Center)
+	[
+		ContainerWidget
+	];
+}
+
+TSharedRef<SInputChord> SInputChord::MakeForInputChord(
+	const FInputChord& InInputChord,
+	const FText& InActionLabel)
+{
+	return SNew(SInputChord)
+		.Icon(GetIconForInputChord(InInputChord))
+		.InputLabel(InInputChord.GetInputText())
+		.ActionLabel(InActionLabel);
+}
+
+TSharedRef<SInputChord> SInputChord::MakeForCommandInfo(const FUICommandInfo& InCommandInfo, const EMultipleKeyBindingIndex InKeyBindingIndex)
+{
+	const FSlateBrush* IconBrush = nullptr;
+
+	if (ensureMsgf(
+		InCommandInfo.GetActiveChord(InKeyBindingIndex)->IsValidChord(),
+		TEXT("CommandInfo '%s' has no valid input chord"), *InCommandInfo.GetLabel().ToString()))
+	{
+		IconBrush = GetIconForInputChord(*InCommandInfo.GetActiveChord(InKeyBindingIndex));
+	}
+
+	return SNew(SInputChord)
+		.Icon(IconBrush)
+		.InputLabel(InCommandInfo.GetInputText())
+		.ActionLabel(InCommandInfo.GetLabel())
+		.ToolTipText(InCommandInfo.GetDescription());
+}
+
+const FSlateBrush* SInputChord::GetIconForInputChord(const FInputChord& InInputChord)
+{
+	return FAppStyle::GetBrush(EKeys::GetMenuCategoryPaletteIcon(InInputChord.Key.GetMenuCategory()));
+}
+
+END_SLATE_FUNCTION_BUILD_OPTIMIZATION
diff --git a/Engine/Source/Editor/EditorWidgets/Public/SInputChord.h b/Engine/Source/Editor/EditorWidgets/Public/SInputChord.h
new file mode 100644
--- /dev/null
+++ b/Engine/Source/Editor/EditorWidgets/Public/SInputChord.h
@@ -0,0 +1,47 @@
+// Copyright Epic Games, Inc. All Rights Reserved.
+
+#pragma once
+
+#include "Framework/Commands/UICommandInfo.h"
+#include "Styling/CoreStyle.h"
+#include "Styling/SlateTypes.h"
+#include "Styling/SlateWidgetStyleAsset.h"
+#include "Widgets/SCompoundWidget.h"
+
+/**
+ * Displays a given input chord in a user-friendly way.
+ * ie. "[Keyboard Icon] Ctrl + C | Copy"
+ */
+class EDITORWIDGETS_API SInputChord : public SCompoundWidget
+{
+public:
+	SLATE_BEGIN_ARGS(SInputChord)
+		: _InputLabelStyle(&FCoreStyle::Get().GetWidgetStyle<FTextBlockStyle>(TEXT("NormalText")))
+		, _ActionLabelStyle(&FCoreStyle::Get().GetWidgetStyle<FTextBlockStyle>(TEXT("NormalText")))
+	{ }
+		/** The Icon should be derived from an FInputChord, usually a keyboard or mouse icon. */
+		SLATE_ARGUMENT(const FSlateBrush*, Icon)
+		/** Label for the input key combination, ie. "Ctrl + C" */
+		SLATE_ARGUMENT(FText, InputLabel)
+		/** Optionally specify the style of the input label */
+		SLATE_STYLE_ARGUMENT(FTextBlockStyle, InputLabelStyle)
+		/** Label for the resulting action, ie. "Copy" */
+		SLATE_ARGUMENT(FText, ActionLabel)
+		/** Optionally specify the style of the action label */
+		SLATE_STYLE_ARGUMENT(FTextBlockStyle, ActionLabelStyle)
+		/** Optionally override the delimiter widget between the input label and action label (if set). By default, this is a vertical separator. */
+		SLATE_ARGUMENT(TSharedPtr<SWidget>, InputLabelDelimiterOverride)
+	SLATE_END_ARGS()
+
+	/** Constructs this widget with InArgs */
+	void Construct(const FArguments& InArgs);
+
+	/** Create an SInputChord widget for the given InputChord and (optional) ActionLabel. The InputLabel (ie. "Ctrl + C") is derived from the InputChord. */
+	static TSharedRef<SInputChord> MakeForInputChord(const FInputChord& InInputChord, const FText& InActionLabel = { });
+
+	/** Create an SInputChord widget for the given CommandInfo. Optionally, provide a KeyBindingIndex for commands with multiple bindings. */
+	static TSharedRef<SInputChord> MakeForCommandInfo(const FUICommandInfo& InCommandInfo, const EMultipleKeyBindingIndex InKeyBindingIndex = EMultipleKeyBindingIndex::Primary);
+
+private:
+	static const FSlateBrush* GetIconForInputChord(const FInputChord& InInputChord);
+};
