[Truncation Warnings] Fix and enable truncation warnings for EditorConfig and EditorStyle and EditorWidgets modules

diff --git a/Engine/Source/Editor/EditorStyle/Private/StarshipStyle.h b/Engine/Source/Editor/EditorStyle/Private/StarshipStyle.h
--- a/Engine/Source/Editor/EditorStyle/Private/StarshipStyle.h
+++ b/Engine/Source/Editor/EditorStyle/Private/StarshipStyle.h
@@ -85,26 +85,26 @@ public:
 		static void SetColor(const TSharedRef<FLinearColor>& Source, const FLinearColor& Value);
 		static bool IncludeEditorSpecificStyles();
 
-		const FVector2D Icon7x16;
-		const FVector2D Icon8x4;
-		const FVector2D Icon16x4;
-		const FVector2D Icon8x8;
-		const FVector2D Icon10x10;
-		const FVector2D Icon12x12;
-		const FVector2D Icon12x16;
-		const FVector2D Icon14x14;
-		const FVector2D Icon16x16;
-		const FVector2D Icon16x20;
-		const FVector2D Icon20x20;
-		const FVector2D Icon22x22;
-		const FVector2D Icon24x24;
-		const FVector2D Icon25x25;
-		const FVector2D Icon32x32;
-		const FVector2D Icon40x40;
-		const FVector2D Icon48x48;
-		const FVector2D Icon64x64;
-		const FVector2D Icon36x24;
-		const FVector2D Icon128x128;
+		static const FVector2f Icon7x16;
+		static const FVector2f Icon8x4;
+		static const FVector2f Icon16x4;
+		static const FVector2f Icon8x8;
+		static const FVector2f Icon10x10;
+		static const FVector2f Icon12x12;
+		static const FVector2f Icon12x16;
+		static const FVector2f Icon14x14;
+		static const FVector2f Icon16x16;
+		static const FVector2f Icon16x20;
+		static const FVector2f Icon20x20;
+		static const FVector2f Icon22x22;
+		static const FVector2f Icon24x24;
+		static const FVector2f Icon25x25;
+		static const FVector2f Icon32x32;
+		static const FVector2f Icon40x40;
+		static const FVector2f Icon48x48;
+		static const FVector2f Icon64x64;
+		static const FVector2f Icon36x24;
+		static const FVector2f Icon128x128;
 
 		// These are the colors that are updated by the user style customizations
 		const TSharedRef< FLinearColor > SelectionColor_Subdued_LinearRef;
diff --git a/Engine/Source/Editor/EditorWidgets/EditorWidgets.Build.cs b/Engine/Source/Editor/EditorWidgets/EditorWidgets.Build.cs
--- a/Engine/Source/Editor/EditorWidgets/EditorWidgets.Build.cs
+++ b/Engine/Source/Editor/EditorWidgets/EditorWidgets.Build.cs
@@ -6,6 +6,9 @@ public class EditorWidgets : ModuleRules
 {
 	public EditorWidgets(ReadOnlyTargetRules Target) : base(Target)
 	{
+		// Enable truncation warnings in this module
+		CppCompileWarningSettings.UnsafeTypeCastWarningLevel = WarningLevel.Warning;
+
 		PrivateIncludePathModuleNames.Add("AssetRegistry");
 
 		PublicDependencyModuleNames.AddRange(
diff --git a/Engine/Source/Editor/EditorWidgets/Private/SAssetDiscoveryIndicator.cpp b/Engine/Source/Editor/EditorWidgets/Private/SAssetDiscoveryIndicator.cpp
--- a/Engine/Source/Editor/EditorWidgets/Private/SAssetDiscoveryIndicator.cpp
+++ b/Engine/Source/Editor/EditorWidgets/Private/SAssetDiscoveryIndicator.cpp
@@ -155,12 +155,12 @@ void SAssetDiscoveryIndicator::OnAssetRegistryFileLoadProgress(const IAssetRegis
 	if (ProgressUpdateData.bIsDiscoveringAssetFiles)
 	{
 		// Marquee while we're discovering asset files as we can't yet show an accurate percentage
-		Progress = TOptional<float>();
+		Progress.Reset();
 
 		// Once we have discovered enough assets, show a preview of the progress bar so it doesn't jump from 0 to 50%
 		if (ProgressUpdateData.NumAssetsProcessedByAssetRegistry > 10000 && ProgressUpdateData.NumTotalAssets > 0)
 		{
-			Progress = ProgressUpdateData.NumAssetsProcessedByAssetRegistry / (float)ProgressUpdateData.NumTotalAssets;
+			Progress = static_cast<float>(ProgressUpdateData.NumAssetsProcessedByAssetRegistry) / static_cast<float>(ProgressUpdateData.NumTotalAssets);
 		}
 
 		MainStatusText = LOCTEXT("DiscoveringAssetFiles", "Discovering Asset Files");
@@ -170,7 +170,7 @@ void SAssetDiscoveryIndicator::OnAssetRegistryFileLoadProgress(const IAssetRegis
 	{
 		if (ProgressUpdateData.NumTotalAssets > 0)
 		{
-			Progress = ProgressUpdateData.NumAssetsProcessedByAssetRegistry / (float)ProgressUpdateData.NumTotalAssets;
+			Progress = static_cast<float>(ProgressUpdateData.NumAssetsProcessedByAssetRegistry) / static_cast<float>(ProgressUpdateData.NumTotalAssets);
 		}
 		else
 		{
diff --git a/Engine/Source/Editor/EditorWidgets/Private/SAssetSearchBox.cpp b/Engine/Source/Editor/EditorWidgets/Private/SAssetSearchBox.cpp
--- a/Engine/Source/Editor/EditorWidgets/Private/SAssetSearchBox.cpp
+++ b/Engine/Source/Editor/EditorWidgets/Private/SAssetSearchBox.cpp
@@ -100,11 +100,11 @@ void SAssetSearchBox::Construct( const FArguments& InArgs )
 				(
 				SNew(SBorder)
 				.BorderImage(FAppStyle::GetBrush("Menu.Background"))
-				.Padding( FMargin(2) )
+				.Padding( FMargin(2.f) )
 				[
 					SNew(SBox)
-					.MinDesiredWidth(175)	// to enforce some minimum width, ideally we define the minimum, not a fixed width
-					.HeightOverride(250)	// avoids flickering, ideally this would be adaptive to the content without flickering
+					.MinDesiredWidth(175.f)	// to enforce some minimum width, ideally we define the minimum, not a fixed width
+					.HeightOverride(250.f)	// avoids flickering, ideally this would be adaptive to the content without flickering
 					[
 						SAssignNew(SuggestionListView, SListView< TSharedPtr<FSuggestionListEntry> >)
 						.ListItemsSource(&Suggestions)
diff --git a/Engine/Source/Editor/EditorWidgets/Private/SEnumCombo.cpp b/Engine/Source/Editor/EditorWidgets/Private/SEnumCombo.cpp
--- a/Engine/Source/Editor/EditorWidgets/Private/SEnumCombo.cpp
+++ b/Engine/Source/Editor/EditorWidgets/Private/SEnumCombo.cpp
@@ -42,7 +42,7 @@ void SEnumComboBox::Construct(const FArguments& InArgs, const UEnum* InEnum)
 		for (int32 i = 0; i < Enum->NumEnums() - 1; i++)
 		{
 			// Note: SEnumComboBox API prior to bitflags only supports 32 bit values, truncating the value here to keep the old API.
-			const int32 Value = Enum->GetValueByIndex(i);
+			const int32 Value = static_cast<int32>(Enum->GetValueByIndex(i));
 			EnumValueSubset.Add(Value);
 		}
 	}
@@ -94,7 +94,7 @@ void SEnumComboBox::Construct(const FArguments& InArgs, const UEnum* InEnum)
 			const int32 Index = Enum->GetIndexByValue(Value);
 			if (Enum->HasMetaData(TEXT("Hidden"), Index) == false)
 			{
-				VisibleEnums.Emplace(Index, Enum->GetValueByIndex(Index), Enum->GetDisplayNameTextByIndex(Index), Enum->GetToolTipTextByIndex(Index));
+				VisibleEnums.Emplace(Index, static_cast<int32>(Enum->GetValueByIndex(Index)), Enum->GetDisplayNameTextByIndex(Index), Enum->GetToolTipTextByIndex(Index));
 			}
 		}
 	}
diff --git a/Engine/Source/Editor/EditorWidgets/Private/STextPropertyEditableTextBox.cpp b/Engine/Source/Editor/EditorWidgets/Private/STextPropertyEditableTextBox.cpp
--- a/Engine/Source/Editor/EditorWidgets/Private/STextPropertyEditableTextBox.cpp
+++ b/Engine/Source/Editor/EditorWidgets/Private/STextPropertyEditableTextBox.cpp
@@ -60,11 +60,11 @@ void STextPropertyEditableStringTableReference::Construct(const FArguments& InAr
 	TSharedRef<SHorizontalBox> HorizontalBox = SNew(SHorizontalBox);
 
 	HorizontalBox->AddSlot()
-		.Padding(0)
+		.Padding(0.f)
 		[
 			SAssignNew(StringTableOptionsCombo, SComboButton)
 			.ComboButtonStyle(&InArgs._ComboStyle->ComboButtonStyle)
-			.ContentPadding(FMargin(4.0, 2.0))
+			.ContentPadding(FMargin(4.0f, 2.0f))
 			.OnGetMenuContent(this, &STextPropertyEditableStringTableReference::OnGetStringTableComboOptions)
 			.OnComboBoxOpened(this, &STextPropertyEditableStringTableReference::UpdateStringTableComboOptions)
 			.CollapseMenuOnParentFocus(true)
@@ -78,11 +78,11 @@ void STextPropertyEditableStringTableReference::Construct(const FArguments& InAr
 		];
 
 	HorizontalBox->AddSlot()
-		.Padding(10, 0)
+		.Padding(10.f, 0.f)
 		[
 			SAssignNew(StringTableKeysCombo, SComboButton)
 			.ComboButtonStyle(&InArgs._ComboStyle->ComboButtonStyle)
-			.ContentPadding(FMargin(4.0, 2.0))
+			.ContentPadding(FMargin(4.0f, 2.0f))
 			.IsEnabled(this, &STextPropertyEditableStringTableReference::IsUnlinkEnabled)
 			.OnGetMenuContent(this, &STextPropertyEditableStringTableReference::OnGetStringTableKeyOptions)
 			.OnComboBoxOpened(this, &STextPropertyEditableStringTableReference::UpdateStringTableKeyOptions)
@@ -105,12 +105,12 @@ void STextPropertyEditableStringTableReference::Construct(const FArguments& InAr
 				SNew(SBox)
 				.HAlign(HAlign_Center)
 				.VAlign(VAlign_Center)
-				.HeightOverride(22)
-				.WidthOverride(22)
+				.HeightOverride(22.f)
+				.WidthOverride(22.f)
 				[
 					SNew(SButton)
 					.ButtonStyle(InArgs._ButtonStyle)
-					.ContentPadding(0)
+					.ContentPadding(0.f)
 					.ToolTipText(LOCTEXT("UnlinkStringTable", "Unlink"))
 					.IsEnabled(this, &STextPropertyEditableStringTableReference::IsUnlinkEnabled)
 					.OnClicked(this, &STextPropertyEditableStringTableReference::OnUnlinkClicked)
@@ -166,7 +166,7 @@ TSharedRef<SWidget> STextPropertyEditableStringTableReference::OnGetStringTableC
 			.WidgetIndex_Lambda([this]() { return StringTableComboOptions.IsEmpty() ? 0 : 1; })
 
 				+SWidgetSwitcher::Slot() // Appears when there are no string tables with keys
-				.Padding(12)
+				.Padding(12.f)
 				[
 					SNew(STextBlock).Text(LOCTEXT("EmptyStringTableList", "No string tables available"))
 				]
@@ -174,14 +174,14 @@ TSharedRef<SWidget> STextPropertyEditableStringTableReference::OnGetStringTableC
 				+SWidgetSwitcher::Slot() // Appears when there's a string table with at least a key
 				[
 					SNew(SBox)
-					.Padding(4)
-					.WidthOverride(280)
-					.MaxDesiredHeight(600)
+					.Padding(4.f)
+					.WidthOverride(280.f)
+					.MaxDesiredHeight(600.f)
 					[
 						SNew(SVerticalBox)
 						+SVerticalBox::Slot()
 						.FillHeight(1.f)
-						.Padding(0, 5, 0, 0)
+						.Padding(0.f, 5.f, 0.f, 0.f)
 						[
 							SAssignNew(StringTableOptionsList, SListView<TSharedPtr<FAvailableStringTable>>)
 							.ListItemsSource(&StringTableComboOptions)
@@ -213,9 +213,9 @@ TSharedRef<SWidget> STextPropertyEditableStringTableReference::OnGetStringTableK
 		.Padding(ComboButtonStyle.MenuBorderPadding)
 		[
 			SNew(SBox)
-			.Padding(4)
-			.WidthOverride(280)
-			.MaxDesiredHeight(600)
+			.Padding(4.f)
+			.WidthOverride(280.f)
+			.MaxDesiredHeight(600.f)
 			[
 				SNew(SVerticalBox)
 				+SVerticalBox::Slot()
@@ -226,7 +226,7 @@ TSharedRef<SWidget> STextPropertyEditableStringTableReference::OnGetStringTableK
 				]
 				+SVerticalBox::Slot()
 				.FillHeight(1.f)
-				.Padding(0, 5, 0, 0)
+				.Padding(0.f, 5.f, 0.f, 0.f)
 				[
 					SAssignNew(StringTableKeysList, SListView<TSharedPtr<FString>>)
 					.ListItemsSource(&KeyComboOptions)
@@ -617,12 +617,12 @@ public:
 			[
 				SNew(SBorder)
 				.BorderImage(FAppStyle::Get().GetBrush("DetailsView.GridLine"))
-				.Padding(FMargin(0, 0, 0, 1))
+				.Padding(FMargin(0.f, 0.f, 0.f, 1.f))
 				[
 					SNew(SBorder)
 					.BorderImage(FAppStyle::Get().GetBrush("DetailsView.CategoryTop"))
 					.BorderBackgroundColor(FSlateColor(FLinearColor::White))
-					.Padding(FMargin(12, 8, 0, 8))
+					.Padding(FMargin(12.f, 8.f, 0.f, 8.f))
 					.VAlign(VAlign_Center)
 					[
 						SNew(STextBlock)
@@ -647,12 +647,12 @@ public:
 					[
 						SNew(SBorder)
 						.BorderImage(FAppStyle::Get().GetBrush("DetailsView.GridLine"))
-						.Padding(FMargin(0, 0, 0, 1))
+						.Padding(FMargin(0.f, 0.f, 0.f, 1.f))
 						[
 							SNew(SBorder)
 							.BorderImage(FAppStyle::Get().GetBrush("DetailsView.CategoryMiddle"))
 							.BorderBackgroundColor(this, &STextPropertyEditableOptionRow::GetBackgroundColor)
-							.Padding(FMargin(20, 3.5, 0, 3.5))
+							.Padding(FMargin(20.f, 3.5f, 0.f, 3.5f))
 							.HAlign(HAlign_Left)
 							.VAlign(VAlign_Center)
 							[
@@ -668,12 +668,12 @@ public:
 				[
 					SNew(SBorder)
 					.BorderImage(FAppStyle::Get().GetBrush("DetailsView.GridLine"))
-					.Padding(FMargin(0, 0, 0, 1))
+					.Padding(FMargin(0.f, 0.f, 0.f, 1.f))
 					[
 						SNew(SBorder)
 						.BorderImage(FAppStyle::Get().GetBrush("DetailsView.CategoryMiddle"))
 						.BorderBackgroundColor(this, &STextPropertyEditableOptionRow::GetBackgroundColor)
-						.Padding(FMargin(14, 3.5, 4, 3.5))
+						.Padding(FMargin(14.f, 3.5f, 4.f, 3.5f))
 						.HAlign(InArgs._ContentHAlign)
 						.VAlign(VAlign_Center)
 						[
@@ -785,7 +785,7 @@ void STextPropertyEditableTextBox::Construct(const FArguments& InArgs, const TSh
 			.MenuContent()
 			[
 				SNew(SBox)
-				.WidthOverride(340)
+				.WidthOverride(340.f)
 				.Padding(1)
 				[
 					SNew(SVerticalBox)
