ContentBrowser: Fix a bug where a disabled filter would enable again when restarting the editor

This was due to LoadSettings() being called twice and not handling being called multiple times probably.

diff --git a/Engine/Source/Editor/ContentBrowser/Private/SContentBrowser.cpp b/Engine/Source/Editor/ContentBrowser/Private/SContentBrowser.cpp
--- a/Engine/Source/Editor/ContentBrowser/Private/SContentBrowser.cpp
+++ b/Engine/Source/Editor/ContentBrowser/Private/SContentBrowser.cpp
@@ -390,9 +390,6 @@ void SContentBrowser::Construct( const FArguments& InArgs, const FName& InInstan
 	// Set the initial history data
 	HistoryManager.AddHistoryData();
 
-	// Load settings if they were specified
-	LoadSettings(InInstanceName);
-	
 	// We want to be able to search the feature packs in the super search so we need the module loaded 
 	IAddContentDialogModule& AddContentDialogModule = FModuleManager::LoadModuleChecked<IAddContentDialogModule>("AddContentDialog");
 
diff --git a/Engine/Source/Editor/EditorWidgets/Public/Filters/SAssetFilterBar.h b/Engine/Source/Editor/EditorWidgets/Public/Filters/SAssetFilterBar.h
--- a/Engine/Source/Editor/EditorWidgets/Public/Filters/SAssetFilterBar.h
+++ b/Engine/Source/Editor/EditorWidgets/Public/Filters/SAssetFilterBar.h
@@ -462,10 +462,10 @@ protected:
 		// Load all the type filters
 		for(const TSharedRef<FCustomClassFilterData> &CustomClassFilter : this->CustomClassFilters)
 		{
+			const FString FilterName = CustomClassFilter->GetFilterName();
+			
 			if(!this->IsClassTypeInUse(CustomClassFilter))
 			{
-				const FString FilterName = CustomClassFilter->GetFilterName();
-
 				// Try to find this type filter in our list of saved filters
 				if ( const bool* bIsActive = FilterBarConfig->TypeFilters.Find(FilterName) )
 				{
@@ -476,6 +476,11 @@ protected:
 					UnknownTypeFilters.Remove(FilterName);
 				}
 			}
+			else
+			{
+				// If the filter is already in use, we don't want to add a duplicate but it isn't unknown - so remove it from the UnknownTypeFilters
+				UnknownTypeFilters.Remove(FilterName);
+			}
 		}
 	}
 	
