[NamingTokens] Slate customization (via existing FTemplateString)

diff --git a/Engine/Source/Editor/DetailCustomizations/Private/TemplateStringStructCustomization.cpp b/Engine/Source/Editor/DetailCustomizations/Private/TemplateStringStructCustomization.cpp
--- a/Engine/Source/Editor/DetailCustomizations/Private/TemplateStringStructCustomization.cpp
+++ b/Engine/Source/Editor/DetailCustomizations/Private/TemplateStringStructCustomization.cpp
@@ -23,6 +23,11 @@ void FTemplateStringStructCustomization::CustomizeHeader(TSharedRef<IPropertyHan
 	TemplateStringProperty = StructPropertyHandle->GetChildHandle(GET_MEMBER_NAME_CHECKED(FTemplateString, Template));
 	check(TemplateStringProperty);
 
+	ResolvedStringProperty = StructPropertyHandle->GetChildHandle(GET_MEMBER_NAME_CHECKED(FTemplateString, Resolved));
+	check(ResolvedStringProperty);
+
+	StructPropertyHandle->SetOnPropertyResetToDefault(FSimpleDelegate::CreateSP(this, &FTemplateStringStructCustomization::OnResetToDefault));
+
 	HeaderRow
 	.NameContent()
 	[
@@ -33,11 +38,9 @@ void FTemplateStringStructCustomization::CustomizeHeader(TSharedRef<IPropertyHan
 	.MinDesiredWidth(125.0f)
 	[
 		SNew(STemplateStringEditableTextBox)
-		.ToolTipText_Lambda([ToolTipText = GetToolTip()]()
-		{
-			return ToolTipText;
-		})
+		.ToolTipText_Raw(this, &FTemplateStringStructCustomization::GetToolTip)
 		.Text_Raw(this, &FTemplateStringStructCustomization::GetText)
+		.ResolvedText_Raw(this, &FTemplateStringStructCustomization::GetResolvedText)
 		.OnTextChanged(this, &FTemplateStringStructCustomization::SetText)
 	];
 }
@@ -55,6 +58,23 @@ FText FTemplateStringStructCustomization::GetText() const
 	return FText::GetEmpty();
 }
 
+FText FTemplateStringStructCustomization::GetResolvedText() const
+{
+	if(ResolvedStringProperty.IsValid())
+	{
+		FText ResolvedString;
+		ResolvedStringProperty->GetValue(ResolvedString);
+
+		if(!ResolvedString.IsEmpty())
+		{
+			return ResolvedString;
+		}
+	}
+
+	// If the ResolvedText isn't available, or is empty, we fallback to the original template/tokenized string
+	return GetText();
+}
+
 void FTemplateStringStructCustomization::SetText(const FText& InNewText) const
 {
 	if(TemplateStringProperty.IsValid())
@@ -119,4 +139,17 @@ const TArray<FString>& FTemplateStringStructCustomization::GetValidArguments() c
 	return EmptyArray;
 }
 
+void FTemplateStringStructCustomization::OnResetToDefault()
+{
+	if (TemplateStringProperty.IsValid())
+	{
+		TemplateStringProperty->ResetToDefault();
+	}
+
+	if (ResolvedStringProperty.IsValid())
+	{
+		ResolvedStringProperty->ResetToDefault();
+	}
+}
+
 #undef LOCTEXT_NAMESPACE
diff --git a/Engine/Source/Editor/DetailCustomizations/Private/TemplateStringStructCustomization.h b/Engine/Source/Editor/DetailCustomizations/Private/TemplateStringStructCustomization.h
--- a/Engine/Source/Editor/DetailCustomizations/Private/TemplateStringStructCustomization.h
+++ b/Engine/Source/Editor/DetailCustomizations/Private/TemplateStringStructCustomization.h
@@ -35,6 +35,9 @@ private:
 	/** Get the property string value (as FText). */
 	FText GetText() const;
 
+	/** Get the property string value (as FText). */
+	FText GetResolvedText() const;
+
 	/** Set the property string value (from FText). */
 	void SetText(const FText& InNewText) const;
 
@@ -44,10 +47,16 @@ private:
 	/** Get the list of valid arguments. */
 	const TArray<FString>& GetValidArguments() const;
 
+	/** Custom behavior that clears both the Template and Resolved strings. */
+	void OnResetToDefault();
+
 private:
-	/** Pointer to the string that will be set when changing the path */
+	/** Pointer to the string that will be set when changing the path. */
 	TSharedPtr<IPropertyHandle> TemplateStringProperty;
 
+	/** Pointer to the stored result of the template string after it's resolved and all args are replaced. */
+	TSharedPtr<IPropertyHandle> ResolvedStringProperty;
+
 	/** Cached tooltip, created by GetToolTip(). */
 	mutable FText CachedTooltip;
 
diff --git a/Engine/Source/Editor/EditorWidgets/Private/STemplateStringEditableTextBox.cpp b/Engine/Source/Editor/EditorWidgets/Private/STemplateStringEditableTextBox.cpp
--- a/Engine/Source/Editor/EditorWidgets/Private/STemplateStringEditableTextBox.cpp
+++ b/Engine/Source/Editor/EditorWidgets/Private/STemplateStringEditableTextBox.cpp
@@ -4,26 +4,377 @@
 
 #include "EditorWidgetsStyle.h"
 #include "HAL/PlatformCrt.h"
+#include "Internationalization/Regex.h"
 #include "Internationalization/Text.h"
 #include "Layout/Margin.h"
 #include "SlateOptMacros.h"
 #include "Styling/ISlateStyle.h"
 #include "Styling/SlateTypes.h"
+#include "Styling/StyleColors.h"
 #include "TemplateStringSyntaxHighlighterMarshaller.h"
+#include "Widgets/Layout/SBox.h"
+#include "Widgets/Notifications/SPopUpErrorText.h"
+#include "Widgets/SBoxPanel.h"
+
+#define LOCTEXT_NAMESPACE "STemplateStringEditableTextBox"
 
 BEGIN_SLATE_FUNCTION_BUILD_OPTIMIZATION
 
 void STemplateStringEditableTextBox::Construct(const FArguments& InArgs)
 {
-	SMultiLineEditableTextBox::Construct(SMultiLineEditableTextBox::FArguments()
-		.Style(&FEditorWidgetsStyle::Get().GetWidgetStyle<FEditableTextBoxStyle>("NormalEditableTextBox"))
-		.Text(InArgs._Text)
-		.Marshaller(FTemplateStringSyntaxHighlighterMarshaller::Create(FTemplateStringSyntaxHighlighterMarshaller::FSyntaxTextStyle()))
-		.AllowMultiLine(false)
-		.OverflowPolicy(ETextOverflowPolicy::Ellipsis)
-		.Margin(0.0f)
-		.OnTextChanged(InArgs._OnTextChanged)
-	);
+	if (InArgs._Style == nullptr)
+	{
+		SetStyle(&FEditorWidgetsStyle::Get().GetWidgetStyle<FEditableTextBoxStyle>("NormalEditableTextBox"));		
+	}
+	else
+	{
+		SetStyle(InArgs._Style);
+	}
+
+	static const FVector2D ScrollBarThickness(9.0f, 9.0f);
+	static const FMargin ErrorWidgetPadding(3, 0); 
+
+	TokenizedText = InArgs._Text;
+	ResolvedText = InArgs._ResolvedText;
+	ValidArguments = InArgs._ValidArguments;
+
+	OnValidateTokenizedText = InArgs._OnValidateTokenizedText;
+	OnTokenizedTextChanged = InArgs._OnTextChanged;
+	OnTokenizedTextCommitted = InArgs._OnTextCommitted;
+
+	ErrorReporting = SNew(SPopupErrorText);
+	ErrorReporting->SetError(FText::GetEmpty());
+
+	ChildSlot
+	[
+		SNew(SBorder)
+		.BorderImage(this, &STemplateStringEditableTextBox::GetBorderImage)
+		.BorderBackgroundColor(TextBoxStyle->BackgroundColor)
+		.ForegroundColor(this, &SMultiLineEditableTextBox::GetForegroundColor)
+		.Padding(TextBoxStyle->Padding)
+		[
+			SAssignNew(Box, SHorizontalBox)
+
+			+ SHorizontalBox::Slot()
+			.HAlign(HAlign_Fill)
+			.VAlign(VAlign_Fill)
+			.AutoWidth()
+			.Padding(4, 0)
+			[
+				SNew(STextBlock)
+				.Text(FText::FromString(TEXT("{ }")))
+				.ColorAndOpacity(FSlateColor(FStyleColors::AccentPurple))
+			]
+
+			+ SHorizontalBox::Slot()
+			.HAlign(HAlign_Fill)
+			.VAlign(VAlign_Fill)
+			.FillWidth(1.0f)
+			[
+				SNew(SVerticalBox)
+
+				+ SVerticalBox::Slot()
+				.HAlign(HAlign_Fill)
+				.VAlign(VAlign_Fill)
+				.FillHeight(1.0f)
+				[
+					SAssignNew(EditableText, SMultiLineEditableText)
+					.Text(this, &STemplateStringEditableTextBox::GetText)
+					.TextStyle(&TextBoxStyle->TextStyle)
+					.Marshaller(FTemplateStringSyntaxHighlighterMarshaller::Create(FTemplateStringSyntaxHighlighterMarshaller::FSyntaxTextStyle()))
+					.AllowMultiLine(InArgs._AllowMultiLine)
+					.OverflowPolicy(ETextOverflowPolicy::Ellipsis)
+					.Margin(0.0f)
+					.OnTextChanged(this, &STemplateStringEditableTextBox::OnEditableTextChanged)
+					.OnTextCommitted(this, &STemplateStringEditableTextBox::OnEditableTextCommitted)
+				]
+
+				+ SVerticalBox::Slot()
+				.AutoHeight()
+				[
+					SAssignNew(HScrollBarBox, SBox)
+					.Padding(TextBoxStyle->HScrollBarPadding)
+					[
+						SAssignNew(HScrollBar, SScrollBar)
+						.Style(&TextBoxStyle->ScrollBarStyle)
+						.Orientation(Orient_Horizontal)
+						.Thickness(ScrollBarThickness)
+					]
+				]
+			]
+
+			+ SHorizontalBox::Slot()
+			.AutoWidth()
+			[
+				SAssignNew(VScrollBarBox, SBox)
+				.Padding(TextBoxStyle->VScrollBarPadding)
+				[
+					SAssignNew(VScrollBar, SScrollBar)
+					.Style(&TextBoxStyle->ScrollBarStyle)
+					.Orientation(Orient_Vertical)
+					.Thickness(ScrollBarThickness)
+				]
+			]
+
+			+ SHorizontalBox::Slot()
+			.AutoWidth()
+			.Padding(ErrorWidgetPadding)
+			[
+				ErrorReporting->AsWidget()
+			]
+		]
+	];
+}
+
+FText STemplateStringEditableTextBox::GetText() const
+{
+	if (HasKeyboardFocus())
+	{
+		return TokenizedText.Get();
+	}
+	else
+	{
+		if (ResolvedText.IsSet())
+		{
+			return ResolvedText.Get();
+		}
+		else
+		{
+			return TokenizedText.Get();
+		}
+	}
+}
+
+const FSlateBrush* STemplateStringEditableTextBox::GetBorderImage() const
+{
+	check(TextBoxStyle);
+	check(EditableText.IsValid());
+
+	if (EditableText->HasKeyboardFocus())
+	{
+		return &TextBoxStyle->BackgroundImageFocused;
+	}
+	else if (EditableText->IsHovered())
+	{
+		return &TextBoxStyle->BackgroundImageHovered;
+	}
+	else
+	{
+		return &TextBoxStyle->BackgroundImageNormal;
+	}
+}
+
+FSlateColor STemplateStringEditableTextBox::GetForegroundColor() const
+{
+	check(TextBoxStyle);
+
+	return HasKeyboardFocus()
+		? TextBoxStyle->FocusedForegroundColor
+		: TextBoxStyle->ForegroundColor;
+}
+
+void STemplateStringEditableTextBox::SetStyle(const FEditableTextBoxStyle* InStyle)
+{
+	if (InStyle)
+	{
+		TextBoxStyle = InStyle;
+	}
+	else
+	{
+		FArguments Defaults;
+		TextBoxStyle = Defaults._Style;
+	}
+
+	check(TextBoxStyle);
+	
+	if (EditableText.IsValid())
+	{
+		EditableText->SetTextStyle(&TextBoxStyle->TextStyle);
+	}
+}
+
+void STemplateStringEditableTextBox::SetError(const FText& InError) const
+{
+	ErrorReporting->SetError(InError);
+}
+
+bool STemplateStringEditableTextBox::HasKeyboardFocus() const
+{
+	// Since keyboard focus is forwarded to our editable text, we will test it instead
+	return SCompoundWidget::HasKeyboardFocus() || EditableText->HasKeyboardFocus();
+}
+
+FReply STemplateStringEditableTextBox::OnFocusReceived(const FGeometry& MyGeometry, const FFocusEvent& InFocusEvent)
+{
+	FReply Reply = FReply::Handled();
+	if ( InFocusEvent.GetCause() != EFocusCause::Cleared )
+	{
+		// Forward keyboard focus to our editable text widget
+		Reply.SetUserFocus(EditableText.ToSharedRef(), InFocusEvent.GetCause());
+	}
+
+	return Reply;
+}
+
+FReply STemplateStringEditableTextBox::OnKeyDown(const FGeometry& MyGeometry, const FKeyEvent& InKeyEvent)
+{
+	const FKey Key = InKeyEvent.GetKey();
+	if (Key == EKeys::Escape && EditableText->HasKeyboardFocus())
+	{
+		// Clear focus
+		return FReply::Handled().SetUserFocus(SharedThis(this), EFocusCause::Cleared);
+	}
+
+	return FReply::Unhandled();
+}
+
+bool STemplateStringEditableTextBox::ValidateTokenizedText(const FText& InTokenizedText)
+{
+	FTextBuilder ErrorMessageBuilder;
+	FText ErrorMessage;
+
+	if (!ValidateTokenBraces(InTokenizedText, ErrorMessage))
+	{
+		ErrorMessageBuilder.AppendLine(ErrorMessage);
+	}
+
+	// Only validate tokens if there were no other errors - it's expensive
+	if (ErrorMessage.IsEmpty() && !ValidateTokenArgs(InTokenizedText, ErrorMessage))
+	{
+		ErrorMessageBuilder.AppendLine(ErrorMessage);
+	}
+
+	if (!ErrorMessageBuilder.IsEmpty())
+	{
+		ErrorMessage = ErrorMessageBuilder.ToText();
+	}
+
+	if (!ErrorMessage.IsEmpty()	&& OnValidateTokenizedText.IsBound())
+	{
+		OnValidateTokenizedText.Execute(InTokenizedText, ErrorMessage);
+	}
+
+	SetError(ErrorMessage);
+
+	// If the error message is empty, the tokenized text is valid
+	return ErrorMessage.IsEmpty();
+}
+
+bool STemplateStringEditableTextBox::ValidateTokenBraces(const FText& InTokenizedText, FText& OutError)
+{
+	// Get number of open and close braces in string
+	int32 OpenBraceCount = 0;
+	int32 CloseBraceCount = 0;
+
+	bool bLastBraceWasOpen = false;
+
+	auto GetUnbalancedBracesMessage = []()
+	{
+		return LOCTEXT("TemplateTextUnbalancedBracingError", "An unbalanced brace was detected. Please ensure that all braces are properly closed.");
+	};
+
+	FString TextString = InTokenizedText.ToString();
+	for (const TCHAR& Char : TextString)
+	{
+		if (Char == TEXT('{'))
+		{
+			++OpenBraceCount;
+
+			// didn't close the last brace
+			if (bLastBraceWasOpen)
+			{
+				OutError = GetUnbalancedBracesMessage();
+				break;
+			}
+
+			bLastBraceWasOpen = true;
+		}
+		else if (Char == TEXT('}'))
+		{
+			++CloseBraceCount;
+			bLastBraceWasOpen = false;
+		}
+	}
+
+	if (OpenBraceCount != CloseBraceCount)
+	{
+		OutError = GetUnbalancedBracesMessage();
+	}
+
+	return OutError.IsEmpty();
+}
+
+bool STemplateStringEditableTextBox::ValidateTokenArgs(const FText& InTokenizedText, FText& OutError)
+{
+	TConstArrayView<FString> ValidArgs = ValidArguments.Get();
+	if (ValidArgs.IsEmpty())
+	{
+		// Always valid if no provided valid args to check against
+		return true;
+	}
+
+	TArray<FString> FoundArgs;
+	if (!ParseArgs(InTokenizedText.ToString(), FoundArgs))
+	{
+		// No args in tokenized text
+		return true;
+	}
+
+	// Not particularly performant
+	for (const FString& FoundArg : FoundArgs)
+	{
+		if (!ValidArgs.ContainsByPredicate([FoundArg](const FString& InValidArg)
+		{
+			return InValidArg.Equals(FoundArg, ESearchCase::IgnoreCase);
+		}))
+		{
+			OutError = LOCTEXT("TemplateTextInvalidArgError", "An argument/token was found that doesn't match any of the provided valid argument names.");
+			return false;
+		}
+	}
+
+	return true;
+}
+
+// @see: UE::NamingTokens::Utils::GetTokenKeysFromString
+bool STemplateStringEditableTextBox::ParseArgs(const FString& InTokenizedTextString, TArray<FString>& OutArgs)
+{
+	TSet<FString, FLocKeySetFuncs> Tokens;
+
+	const FString PatternString = TEXT(R"(\{\s*((?:[a-zA-Z0-9_]+\.)*[a-zA-Z0-9_]+)\s*\})");
+	const FRegexPattern Pattern(PatternString);
+	FRegexMatcher Matcher(Pattern, InTokenizedTextString);
+
+	while (Matcher.FindNext())
+	{
+		FString Token = Matcher.GetCaptureGroup(1);
+		Tokens.Add(Token);
+	}
+
+	OutArgs = Tokens.Array();
+
+	return !OutArgs.IsEmpty();
+}
+
+void STemplateStringEditableTextBox::OnEditableTextChanged(const FText& InTokenizedText)
+{
+	ValidateTokenizedText(InTokenizedText);
+
+	OnTokenizedTextChanged.ExecuteIfBound(InTokenizedText);
+}
+
+void STemplateStringEditableTextBox::OnEditableTextCommitted(const FText& InTokenizedText, ETextCommit::Type InCommitType)
+{
+	ValidateTokenizedText(InTokenizedText);
+
+	if (!TokenizedText.IsBound())
+	{
+		TokenizedText.Set(InTokenizedText);
+	}
+
+	OnTokenizedTextCommitted.ExecuteIfBound(InTokenizedText, InCommitType);
 }
 
 END_SLATE_FUNCTION_BUILD_OPTIMIZATION
+
+#undef LOCTEXT_NAMESPACE
diff --git a/Engine/Source/Editor/EditorWidgets/Public/STemplateStringEditableTextBox.h b/Engine/Source/Editor/EditorWidgets/Public/STemplateStringEditableTextBox.h
--- a/Engine/Source/Editor/EditorWidgets/Public/STemplateStringEditableTextBox.h
+++ b/Engine/Source/Editor/EditorWidgets/Public/STemplateStringEditableTextBox.h
@@ -14,25 +14,96 @@ class FString;
 class FText;
 class ITextLayoutMarshaller;
 
-class EDITORWIDGETS_API STemplateStringEditableTextBox : public SMultiLineEditableTextBox
+class STemplateStringEditableTextBox : public SCompoundWidget
 {
 public:
 	SLATE_BEGIN_ARGS(STemplateStringEditableTextBox)
+		: _Style(nullptr)
+		, _AllowMultiLine(false)
 	{}
+		/** The styling of the textbox. */
+		SLATE_STYLE_ARGUMENT(FEditableTextBoxStyle, Style)
+
 		/** The initial text that will appear in the widget. */
 		SLATE_ATTRIBUTE(FText, Text)
 
-		/** The list of available arguments to use in this template string. */
-		SLATE_ATTRIBUTE(TArray<FString>, ValidArguments)
+		/** The (optional) resolved text, displayed when the text box is not focused. If not provided, the templated text is shown. */
+		SLATE_ATTRIBUTE(FText, ResolvedText)
+
+		/** The list of available arguments to use in this template string. If Empty, any argument name is valid. */
+		SLATE_ATTRIBUTE(TConstArrayView<FString>, ValidArguments)
 
 		/** The marshaller used to get/set the raw text to/from the text layout. */
 		SLATE_ARGUMENT(TSharedPtr<ITextLayoutMarshaller>, Marshaller)
 
-		/** Called whenever the text is changed interactively by the user. */
+		/** Whether to allow multi-line text. */
+		SLATE_ATTRIBUTE(bool, AllowMultiLine)
+
+		/** Allows custom validation. */
+		SLATE_EVENT(FOnVerifyTextChanged, OnValidateTokenizedText)
+
+		/** Called whenever the (tokenized) text is changed interactively by the user. */
 		SLATE_EVENT(FOnTextChanged, OnTextChanged)
-	
+		
+		/** Called whenever the (tokenized) text is committed by the user. */
+		SLATE_EVENT(FOnTextCommitted, OnTextCommitted)
+
 	SLATE_END_ARGS()
 
 	/** Constructs this widget with InArgs */
-	void Construct(const FArguments& InArgs);
+	EDITORWIDGETS_API void Construct(const FArguments& InArgs);
+
+private:
+	/** @return Text contents based on current state */
+	FText GetText() const;
+
+	const FSlateBrush* GetBorderImage() const;
+	EDITORWIDGETS_API virtual FSlateColor GetForegroundColor() const override;
+
+	void SetStyle(const FEditableTextBoxStyle* InStyle);
+
+	void SetError(const FText& InError) const;
+
+	EDITORWIDGETS_API virtual bool HasKeyboardFocus() const override;
+	EDITORWIDGETS_API virtual FReply OnFocusReceived(const FGeometry& MyGeometry, const FFocusEvent& InFocusEvent) override;
+	EDITORWIDGETS_API virtual FReply OnKeyDown(const FGeometry& MyGeometry, const FKeyEvent& InKeyEvent) override;
+
+	bool ValidateTokenizedText(const FText& InTokenizedText);
+	bool ValidateTokenBraces(const FText& InTokenizedText, FText& OutError);
+	bool ValidateTokenArgs(const FText& InTokenizedText, FText& OutError);
+
+	bool ParseArgs(const FString& InTokenizedTextString, TArray<FString>& OutArgs);
+
+	void OnEditableTextChanged(const FText& InTokenizedText);
+	void OnEditableTextCommitted(const FText& InTokenizedText, ETextCommit::Type InCommitType);
+
+private:
+	/** Editable text widget. */
+	TSharedPtr<SMultiLineEditableText> EditableText;
+
+	/** Style shared between Editable and NonEditable text widgets. */
+	const FEditableTextBoxStyle* TextBoxStyle = nullptr;
+
+	TSharedPtr<SHorizontalBox> Box;
+
+	TSharedPtr<SBox> HScrollBarBox;
+	TSharedPtr<SScrollBar> HScrollBar;
+
+	TSharedPtr<SBox> VScrollBarBox;
+	TSharedPtr<SScrollBar> VScrollBar;
+
+	TSharedPtr<class IErrorReportingWidget> ErrorReporting;
+	
+	TAttribute<FText> TokenizedText;
+	TAttribute<FText> ResolvedText;
+	TAttribute<TConstArrayView<FString>> ValidArguments;
+	
+	/** Callback to verify tokenized text when changed. Will return an error message to denote problems. */
+    FOnVerifyTextChanged OnValidateTokenizedText;
+
+	/** Callback when tokenized text is changed. */
+	FOnTextChanged OnTokenizedTextChanged;
+
+	/** Callback when tokenized text is committed. */
+	FOnTextCommitted OnTokenizedTextCommitted;
 };
