Update FilterBar method names

#rnx

diff --git a/Engine/Source/Editor/ContentBrowser/Private/SFilterList.cpp b/Engine/Source/Editor/ContentBrowser/Private/SFilterList.cpp
--- a/Engine/Source/Editor/ContentBrowser/Private/SFilterList.cpp
+++ b/Engine/Source/Editor/ContentBrowser/Private/SFilterList.cpp
@@ -530,7 +530,7 @@ void SFilterList::PopulateAddFilterMenu(UToolMenu* Menu, TSharedPtr<FFilterCateg
                 MenuExpansion->Tooltip,
                 FSlateIcon(FAppStyle::Get().GetStyleSetName(), "PlacementBrowser.Icons.Basic"),
                 FUIAction(
-                FExecuteAction::CreateSP( this, &SFilterList::FilterByTypeCategoryClicked, MenuExpansion, ExpandedCategory->Classes ),
+                FExecuteAction::CreateSP( this, &SFilterList::FilterByTypeCategory, MenuExpansion, ExpandedCategory->Classes ),
                 FCanExecuteAction(),
                 FGetActionCheckState::CreateSP(this, &SFilterList::IsTypeCategoryChecked, MenuExpansion, ExpandedCategory->Classes ) ),
                 EUserInterfaceActionType::ToggleButton
@@ -569,7 +569,7 @@ void SFilterList::PopulateAddFilterMenu(UToolMenu* Menu, TSharedPtr<FFilterCateg
 					CategoryMenuPair.Key->Tooltip,
 					FNewToolMenuDelegate::CreateSP(this, &SFilterList::CreateFiltersMenuCategory, CategoryMenuPair.Value.Classes),
 					FUIAction(
-					FExecuteAction::CreateSP(this, &SFilterList::FilterByTypeCategoryClicked, CategoryMenuPair.Key, CategoryMenuPair.Value.Classes),
+					FExecuteAction::CreateSP(this, &SFilterList::FilterByTypeCategory, CategoryMenuPair.Key, CategoryMenuPair.Value.Classes),
 					FCanExecuteAction(),
 					FGetActionCheckState::CreateSP(this, &SFilterList::IsTypeCategoryChecked, CategoryMenuPair.Key, CategoryMenuPair.Value.Classes)),
 					EUserInterfaceActionType::ToggleButton
diff --git a/Engine/Source/Editor/EditorWidgets/Public/Filters/SAssetFilterBar.h b/Engine/Source/Editor/EditorWidgets/Public/Filters/SAssetFilterBar.h
--- a/Engine/Source/Editor/EditorWidgets/Public/Filters/SAssetFilterBar.h
+++ b/Engine/Source/Editor/EditorWidgets/Public/Filters/SAssetFilterBar.h
@@ -771,7 +771,7 @@ public:
 	}
 
 	/** Handler for when filter by type is selected */
-	void FilterByTypeClicked(TSharedPtr<FCustomClassFilterData> CustomClassFilterData)
+	void FilterByType(TSharedPtr<FCustomClassFilterData> CustomClassFilterData)
 	{
 		if (CustomClassFilterData.IsValid())
 		{
@@ -788,7 +788,7 @@ public:
 	}
 
 	/** Handler for when filter by type category is selected */
-	void FilterByTypeCategoryClicked(TSharedPtr<FFilterCategory> TypeCategory, TArray<TSharedPtr<FCustomClassFilterData>> Classes)
+	void FilterByTypeCategory(TSharedPtr<FFilterCategory> TypeCategory, TArray<TSharedPtr<FCustomClassFilterData>> Classes)
 	{
 		bool bFullCategoryInUse = IsTypeCategoryInUse(TypeCategory, Classes);
 		bool ExecuteOnFilterChanged = false;
@@ -1008,7 +1008,7 @@ protected:
 				FText::Format( LOCTEXT("FilterByTooltipPrefix", "Filter by {0}"), LabelText ),
 				FSlateIconFinder::FindIconForClass(CustomClassFilterData->GetClass()),
 				FUIAction(
-					FExecuteAction::CreateSP( const_cast< SAssetFilterBar<FilterType>* >(this), &SAssetFilterBar<FilterType>::FilterByTypeClicked, CustomClassFilterData ),
+					FExecuteAction::CreateSP( const_cast< SAssetFilterBar<FilterType>* >(this), &SAssetFilterBar<FilterType>::FilterByType, CustomClassFilterData ),
 					FCanExecuteAction(),
 					FIsActionChecked::CreateSP(this, &SAssetFilterBar<FilterType>::IsClassTypeInUse, CustomClassFilterData ) ),
 				EUserInterfaceActionType::ToggleButton
@@ -1088,7 +1088,7 @@ protected:
                 	MenuExpansion->Tooltip,
                 	FSlateIcon(FAppStyle::Get().GetStyleSetName(), "PlacementBrowser.Icons.Basic"),
                 	FUIAction(
-                	FExecuteAction::CreateSP( this, &SAssetFilterBar<FilterType>::FilterByTypeCategoryClicked, MenuExpansion, ExpandedCategory->Classes ),
+                	FExecuteAction::CreateSP( this, &SAssetFilterBar<FilterType>::FilterByTypeCategory, MenuExpansion, ExpandedCategory->Classes ),
                 	FCanExecuteAction(),
                 	FGetActionCheckState::CreateSP(this, &SAssetFilterBar<FilterType>::IsTypeCategoryChecked, MenuExpansion, ExpandedCategory->Classes ) ),
                 	EUserInterfaceActionType::ToggleButton
@@ -1130,7 +1130,7 @@ protected:
 				CategoryMenuPair.Key->Tooltip,
 				FNewToolMenuDelegate::CreateSP(this, &SAssetFilterBar<FilterType>::CreateFiltersMenuCategory, CategoryMenuPair.Value.Classes),
 				FUIAction(
-				FExecuteAction::CreateSP(this, &SAssetFilterBar<FilterType>::FilterByTypeCategoryClicked, CategoryMenuPair.Key, CategoryMenuPair.Value.Classes),
+				FExecuteAction::CreateSP(this, &SAssetFilterBar<FilterType>::FilterByTypeCategory, CategoryMenuPair.Key, CategoryMenuPair.Value.Classes),
 				FCanExecuteAction(),
 				FGetActionCheckState::CreateSP(this, &SAssetFilterBar<FilterType>::IsTypeCategoryChecked, CategoryMenuPair.Key, CategoryMenuPair.Value.Classes)),
 				EUserInterfaceActionType::ToggleButton
diff --git a/Engine/Source/Editor/SceneOutliner/Private/SSceneOutliner.cpp b/Engine/Source/Editor/SceneOutliner/Private/SSceneOutliner.cpp
--- a/Engine/Source/Editor/SceneOutliner/Private/SSceneOutliner.cpp
+++ b/Engine/Source/Editor/SceneOutliner/Private/SSceneOutliner.cpp
@@ -2688,7 +2688,7 @@ void SSceneOutliner::FilterByType(TSharedPtr<FCustomClassFilterData> CustomClass
 {
 	if (FilterBar)
 	{
-		FilterBar->FilterByTypeClicked(CustomClassFilterData);
+		FilterBar->FilterByType(CustomClassFilterData);
 	}
 }
 
@@ -2696,7 +2696,7 @@ void SSceneOutliner::FilterByTypeCategory(TSharedPtr<FFilterCategory> TypeCatego
 {
 	if (FilterBar)
 	{
-		FilterBar->FilterByTypeCategoryClicked(TypeCategory, Classes);
+		FilterBar->FilterByTypeCategory(TypeCategory, Classes);
 	}
 }
 
