Make an invariant text property localizable again when overriding it from its default value

Such overrides are typically intended to be localizable, and if they're not, the override can be manually set back to invariant and future edits to it will stay invariant

diff --git a/Engine/Source/Editor/EditorWidgets/Private/STextPropertyEditableTextBox.cpp b/Engine/Source/Editor/EditorWidgets/Private/STextPropertyEditableTextBox.cpp
--- a/Engine/Source/Editor/EditorWidgets/Private/STextPropertyEditableTextBox.cpp
+++ b/Engine/Source/Editor/EditorWidgets/Private/STextPropertyEditableTextBox.cpp
@@ -1125,6 +1125,10 @@ void STextPropertyEditableTextBox::OnTextCommitted(const FText& NewText, ETextCo
 {
 	const int32 NumTexts = EditableTextProperty->GetNumTexts();
 
+	// Allow the first edit to a default value to make invariant text localizable again, 
+	// as the value being set will become an override that was likely meant to be localized
+	const bool bPreserveCultureInvariance = !EditableTextProperty->IsDefaultValue();
+
 	// Don't commit the Multiple Values text if there are multiple properties being set
 	if (NumTexts > 0 && (NumTexts == 1 || !NewText.ToString().Equals(MultipleValuesText.ToString(), ESearchCase::CaseSensitive)))
 	{
@@ -1160,7 +1164,7 @@ void STextPropertyEditableTextBox::OnTextCommitted(const FText& NewText, ETextCo
 			}
 
 			// Maintain culture invariance when editing the text
-			if (PropertyValue.IsCultureInvariant())
+			if (bPreserveCultureInvariance && PropertyValue.IsCultureInvariant())
 			{
 				EditableTextProperty->SetText(TextIndex, FText::AsCultureInvariant(NewText.ToString()));
 				continue;
