Add IsReadOnly attribute to STemplateStringEditableTextBox which passes the argument along to the underlying Editable Text widget.
#rnx

diff --git a/Engine/Source/Editor/EditorWidgets/Private/STemplateStringEditableTextBox.cpp b/Engine/Source/Editor/EditorWidgets/Private/STemplateStringEditableTextBox.cpp
--- a/Engine/Source/Editor/EditorWidgets/Private/STemplateStringEditableTextBox.cpp
+++ b/Engine/Source/Editor/EditorWidgets/Private/STemplateStringEditableTextBox.cpp
@@ -83,6 +83,7 @@ void STemplateStringEditableTextBox::Construct(const FArguments& InArgs)
 					.TextStyle(&TextBoxStyle->TextStyle)
 					.Marshaller(FTemplateStringSyntaxHighlighterMarshaller::Create(FTemplateStringSyntaxHighlighterMarshaller::FSyntaxTextStyle()))
 					.AllowMultiLine(InArgs._AllowMultiLine)
+					.IsReadOnly(InArgs._IsReadOnly)
 					.OverflowPolicy(ETextOverflowPolicy::Ellipsis)
 					.Margin(0.0f)
 					.OnTextChanged(this, &STemplateStringEditableTextBox::OnEditableTextChanged)
@@ -150,7 +151,11 @@ const FSlateBrush* STemplateStringEditableTextBox::GetBorderImage() const
 	check(TextBoxStyle);
 	check(EditableText.IsValid());
 
-	if (EditableText->HasKeyboardFocus())
+	if (EditableText->IsTextReadOnly())
+	{
+		return &TextBoxStyle->BackgroundImageReadOnly;
+	}
+	else if (EditableText->HasKeyboardFocus())
 	{
 		return &TextBoxStyle->BackgroundImageFocused;
 	}
diff --git a/Engine/Source/Editor/EditorWidgets/Public/STemplateStringEditableTextBox.h b/Engine/Source/Editor/EditorWidgets/Public/STemplateStringEditableTextBox.h
--- a/Engine/Source/Editor/EditorWidgets/Public/STemplateStringEditableTextBox.h
+++ b/Engine/Source/Editor/EditorWidgets/Public/STemplateStringEditableTextBox.h
@@ -1,4 +1,4 @@
-﻿// Copyright Epic Games, Inc. All Rights Reserved.
+// Copyright Epic Games, Inc. All Rights Reserved.
 
 #pragma once
 
@@ -20,6 +20,7 @@ public:
 	SLATE_BEGIN_ARGS(STemplateStringEditableTextBox)
 		: _Style(nullptr)
 		, _AllowMultiLine(false)
+		, _IsReadOnly(false)
 	{}
 		/** The styling of the textbox. */
 		SLATE_STYLE_ARGUMENT(FEditableTextBoxStyle, Style)
@@ -39,6 +40,9 @@ public:
 		/** Whether to allow multi-line text. */
 		SLATE_ATTRIBUTE(bool, AllowMultiLine)
 
+		/** Sets whether this text box can actually be modified interactively by the user */
+		SLATE_ATTRIBUTE(bool, IsReadOnly)
+
 		/** Allows custom validation. */
 		SLATE_EVENT(FOnVerifyTextChanged, OnValidateTokenizedText)
 
