NamingTokens: Fixing issue where STemplateStringEditableTextBox validation function would never run.

diff --git a/Engine/Source/Editor/EditorWidgets/Private/STemplateStringEditableTextBox.cpp b/Engine/Source/Editor/EditorWidgets/Private/STemplateStringEditableTextBox.cpp
--- a/Engine/Source/Editor/EditorWidgets/Private/STemplateStringEditableTextBox.cpp
+++ b/Engine/Source/Editor/EditorWidgets/Private/STemplateStringEditableTextBox.cpp
@@ -254,7 +254,7 @@ bool STemplateStringEditableTextBox::ValidateTokenizedText(const FText& InTokeni
 		ErrorMessage = ErrorMessageBuilder.ToText();
 	}
 
-	if (!ErrorMessage.IsEmpty()	&& OnValidateTokenizedText.IsBound())
+	if (ErrorMessage.IsEmpty()	&& OnValidateTokenizedText.IsBound())
 	{
 		OnValidateTokenizedText.Execute(InTokenizedText, ErrorMessage);
 	}
