Details panel: adding a DetailNameAreaObjectFilter to control which objects are sent to the DetailNameArea

diff --git a/Engine/Source/Editor/EditorWidgets/Private/DetailNameAreaObjectFilterRegistry.cpp b/Engine/Source/Editor/EditorWidgets/Private/DetailNameAreaObjectFilterRegistry.cpp
new file mode 100644
--- /dev/null
+++ b/Engine/Source/Editor/EditorWidgets/Private/DetailNameAreaObjectFilterRegistry.cpp
@@ -0,0 +1,31 @@
+// Copyright Epic Games, Inc. All Rights Reserved.
+
+#include "DetailNameAreaObjectFilterRegistry.h"
+
+#define LOCTEXT_NAMESPACE "EditorWidgets"
+
+namespace UE::EditorWidgets
+{
+
+void FDetailNameAreaObjectFilterRegistry::RegisterDetailNameAreaObjectFilter(TUniquePtr<IDetailNameAreaObjectFilter>&& NewNameAreaFilter)
+{
+	const FName FilterName = NewNameAreaFilter->GetFilterName();
+	ObjectFilters.Add({ FilterName, MoveTemp(NewNameAreaFilter) });
+}
+
+void FDetailNameAreaObjectFilterRegistry::UnregisterDetailNameAreaObjectFilter(const FName& FilterName)
+{
+	ObjectFilters.Remove(FilterName);
+}
+
+void FDetailNameAreaObjectFilterRegistry::ApplyAllFiltersToSelectedObjects(TArray<TWeakObjectPtr<UObject>>& InOutObjects) const
+{
+	for (const TPair<FName, TUniquePtr<IDetailNameAreaObjectFilter>>& Filter : ObjectFilters)
+	{
+		Filter.Value->FilterSelectedObjectsForNameArea(InOutObjects);
+	}
+}
+
+} // end namespace UE::EditorWidgets
+
+#undef LOCTEXT_NAMESPACE
diff --git a/Engine/Source/Editor/EditorWidgets/Private/EditorWidgetsModule.cpp b/Engine/Source/Editor/EditorWidgets/Private/EditorWidgetsModule.cpp
--- a/Engine/Source/Editor/EditorWidgets/Private/EditorWidgetsModule.cpp
+++ b/Engine/Source/Editor/EditorWidgets/Private/EditorWidgetsModule.cpp
@@ -6,6 +6,7 @@
 #include "HAL/Platform.h"
 #include "Modules/ModuleManager.h"
 #include "ObjectNameEditSinkRegistry.h"
+#include "DetailNameAreaObjectFilterRegistry.h"
 #include "SAssetDiscoveryIndicator.h"
 #include "SObjectNameEditableTextBox.h"
 #include "STransportControl.h"
@@ -24,11 +25,13 @@ void FEditorWidgetsModule::StartupModule()
 	FEditorWidgetsStyle::Initialize();
 	
 	ObjectNameEditSinkRegistry = MakeShared<UE::EditorWidgets::FObjectNameEditSinkRegistry>();
+	DetailNameAreaObjectFilterRegistry = MakeShared<UE::EditorWidgets::FDetailNameAreaObjectFilterRegistry>();
 }
 
 void FEditorWidgetsModule::ShutdownModule()
 {
 	ObjectNameEditSinkRegistry.Reset();
+	DetailNameAreaObjectFilterRegistry.Reset();
 
 	FEditorWidgetsStyle::Shutdown();
 }
@@ -56,4 +59,9 @@ TSharedRef<ITransportControl> FEditorWidgetsModule::CreateTransportControl(const
 TSharedRef<UE::EditorWidgets::FObjectNameEditSinkRegistry> FEditorWidgetsModule::GetObjectNameEditSinkRegistry() const
 {
 	return ObjectNameEditSinkRegistry.ToSharedRef();
-}
\ No newline at end of file
+}
+
+TSharedRef<UE::EditorWidgets::FDetailNameAreaObjectFilterRegistry> FEditorWidgetsModule::GetDetailNameAreaObjectFilterRegistry() const
+{
+	return DetailNameAreaObjectFilterRegistry.ToSharedRef();
+}
diff --git a/Engine/Source/Editor/EditorWidgets/Public/DetailNameAreaObjectFilterRegistry.h b/Engine/Source/Editor/EditorWidgets/Public/DetailNameAreaObjectFilterRegistry.h
new file mode 100644
--- /dev/null
+++ b/Engine/Source/Editor/EditorWidgets/Public/DetailNameAreaObjectFilterRegistry.h
@@ -0,0 +1,40 @@
+// Copyright Epic Games, Inc. All Rights Reserved.
+
+#pragma once
+
+#include "Containers/Array.h"
+#include "Templates/SharedPointer.h"
+#include "UObject/WeakObjectPtrTemplates.h"
+
+#define UE_API EDITORWIDGETS_API
+
+class UClass;
+
+namespace UE::EditorWidgets
+{
+
+class IDetailNameAreaObjectFilter
+{
+public:
+	virtual ~IDetailNameAreaObjectFilter() = default;
+	virtual FName GetFilterName() const = 0;
+	virtual void FilterSelectedObjectsForNameArea(TArray<TWeakObjectPtr<UObject>>& InOutObjects) = 0;
+};
+
+class FDetailNameAreaObjectFilterRegistry
+{
+public:
+	
+	UE_API void RegisterDetailNameAreaObjectFilter(TUniquePtr<IDetailNameAreaObjectFilter>&& NewNameAreaFilter);
+	UE_API void UnregisterDetailNameAreaObjectFilter(const FName& FilterName);
+
+	UE_API void ApplyAllFiltersToSelectedObjects(TArray<TWeakObjectPtr<UObject>>& InOutObjects) const;
+
+private:
+
+	TMap<FName, TUniquePtr<IDetailNameAreaObjectFilter>> ObjectFilters;
+};
+
+} // end namespace UE::EditorWidgets
+
+#undef UE_API
diff --git a/Engine/Source/Editor/EditorWidgets/Public/EditorWidgetsModule.h b/Engine/Source/Editor/EditorWidgets/Public/EditorWidgetsModule.h
--- a/Engine/Source/Editor/EditorWidgets/Public/EditorWidgetsModule.h
+++ b/Engine/Source/Editor/EditorWidgets/Public/EditorWidgetsModule.h
@@ -21,6 +21,7 @@ struct FTransportControlArgs;
 namespace UE::EditorWidgets
 {
 class FObjectNameEditSinkRegistry;
+class FDetailNameAreaObjectFilterRegistry;
 }
 
 /** Interface for the widget that wraps an editable text box for viewing the names of objects or editing the labels of actors */
@@ -77,10 +78,14 @@ public:
 	 */
 	virtual TSharedRef<UE::EditorWidgets::FObjectNameEditSinkRegistry> GetObjectNameEditSinkRegistry() const;
 
+	virtual TSharedRef<UE::EditorWidgets::FDetailNameAreaObjectFilterRegistry> GetDetailNameAreaObjectFilterRegistry() const;
+
 	/** Editor Widgets app identifier string */
 	static const FName EditorWidgetsAppIdentifier;
 
 private:
 	/** Registry for name interfaces */
 	TSharedPtr<UE::EditorWidgets::FObjectNameEditSinkRegistry> ObjectNameEditSinkRegistry;
+
+	TSharedPtr<UE::EditorWidgets::FDetailNameAreaObjectFilterRegistry> DetailNameAreaObjectFilterRegistry;
 };
diff --git a/Engine/Source/Editor/PropertyEditor/Private/SDetailNameArea.cpp b/Engine/Source/Editor/PropertyEditor/Private/SDetailNameArea.cpp
--- a/Engine/Source/Editor/PropertyEditor/Private/SDetailNameArea.cpp
+++ b/Engine/Source/Editor/PropertyEditor/Private/SDetailNameArea.cpp
@@ -16,6 +16,7 @@
 #include "Widgets/Input/SButton.h"
 #include "Widgets/Layout/SBox.h"
 #include "Widgets/Text/STextBlock.h"
+#include "DetailNameAreaObjectFilterRegistry.h"
 
 #define LOCTEXT_NAMESPACE "SDetailsView"
 
@@ -63,6 +64,9 @@ void SDetailNameArea::Refresh( const TArray< TWeakObjectPtr<UObject> >& Selected
 		FinalSelectedObjects = SelectedObjects;
 	}
 
+	const FEditorWidgetsModule& EditorWidgetsModule = FModuleManager::LoadModuleChecked<FEditorWidgetsModule>("EditorWidgets");
+	EditorWidgetsModule.GetDetailNameAreaObjectFilterRegistry()->ApplyAllFiltersToSelectedObjects(FinalSelectedObjects);
+
 	ChildSlot
 	[
 		BuildObjectNameArea( FinalSelectedObjects )
