ObjectNameEditSink: adding support for an optional name suffix which can be rendered differently than the name

diff --git a/Engine/Source/Editor/EditorWidgets/Private/SObjectNameEditableTextBox.cpp b/Engine/Source/Editor/EditorWidgets/Private/SObjectNameEditableTextBox.cpp
--- a/Engine/Source/Editor/EditorWidgets/Private/SObjectNameEditableTextBox.cpp
+++ b/Engine/Source/Editor/EditorWidgets/Private/SObjectNameEditableTextBox.cpp
@@ -24,6 +24,7 @@
 #include "Widgets/SCompoundWidget.h"
 #include "Widgets/SWindow.h"
 #include "Widgets/Text/SInlineEditableTextBlock.h"
+#include "Widgets/Text/STextBlock.h"
 
 class FSlateRect;
 struct FSlateBrush;
@@ -45,16 +46,38 @@ void SObjectNameEditableTextBox::Construct( const FArguments& InArgs )
 	Objects = InArgs._Objects;
 	ObjectNameEditSinkRegistry = InArgs._Registry;
 
-	ChildSlot
-	[
-		SAssignNew(TextBox, SInlineEditableTextBlock)
+
+	SAssignNew(TextBox, SInlineEditableTextBlock)
 		.Style(FAppStyle::Get(), "DetailsView.NameTextBlockStyle")
-		.Text( this, &SObjectNameEditableTextBox::GetNameText )
-		.ToolTipText( this, &SObjectNameEditableTextBox::GetNameTooltipText )
-		.Visibility( this, &SObjectNameEditableTextBox::GetNameVisibility )
-		.OnTextCommitted( this, &SObjectNameEditableTextBox::OnNameTextCommitted )
-		.IsReadOnly( this, &SObjectNameEditableTextBox::IsReadOnly )
-		.OnVerifyTextChanged_Static(&FActorEditorUtils::ValidateActorName)
+		.Text(this, &SObjectNameEditableTextBox::GetNameText)
+		.ToolTipText(this, &SObjectNameEditableTextBox::GetNameTooltipText)
+		.Visibility(this, &SObjectNameEditableTextBox::GetNameVisibility)
+		.OnTextCommitted(this, &SObjectNameEditableTextBox::OnNameTextCommitted)
+		.IsReadOnly(this, &SObjectNameEditableTextBox::IsReadOnly)
+		.OnVerifyTextChanged_Static(&FActorEditorUtils::ValidateActorName);
+
+	const bool bShouldUseMiddleEllipsis = true; // GetDefault<UEditorStyleSettings>()->bEnableMiddleEllipsis;
+
+	TSharedPtr<STextBlock> SuffixTextBlock = SNew(STextBlock)
+		.Text(this, &SObjectNameEditableTextBox::GetObjectDisplayNameSuffix)
+		.Visibility(this, &SObjectNameEditableTextBox::GetNameSuffixVisibility)
+		.IsEnabled(false);
+
+	ChildSlot
+	[	
+		SNew(SHorizontalBox)
+			+ SHorizontalBox::Slot()
+			.FillContentWidth(1.0f)
+			.VAlign(VAlign_Center)
+			[
+				TextBox.ToSharedRef()
+			]
+			+ SHorizontalBox::Slot()
+			.FillContentWidth(1000.0f)
+			.VAlign(VAlign_Center)
+			[
+				SuffixTextBlock.ToSharedRef()
+			]
 	];
 }
 
@@ -324,5 +347,46 @@ FText SObjectNameEditableTextBox::GetObjectDisplayName(TWeakObjectPtr<UObject> O
 	return FText::FromString(ObjectPtr->GetName());
 }
 
+EVisibility SObjectNameEditableTextBox::GetNameSuffixVisibility() const
+{
+	if (const TSharedPtr<const UE::EditorWidgets::FObjectNameEditSinkRegistry> RegistryPtr = ObjectNameEditSinkRegistry.Pin())
+	{
+		if (Objects.Num() == 1)
+		{
+			if (const UObject* const ObjectPtr = Objects[0].Get())
+			{
+				const TSharedPtr<const UE::EditorWidgets::IObjectNameEditSink> ObjectNameEditPtr = RegistryPtr->GetObjectNameEditSinkForClass(ObjectPtr->GetClass());
+				if (ObjectNameEditPtr)
+				{
+					return ObjectNameEditPtr->ObjectDisplayNameHasSuffix(ObjectPtr) ? EVisibility::Visible : EVisibility::Collapsed;
+				}
+			}
+		}
+	}
+
+	return EVisibility::Collapsed;
+}
+
+FText SObjectNameEditableTextBox::GetObjectDisplayNameSuffix() const
+{
+	if (const TSharedPtr<const UE::EditorWidgets::FObjectNameEditSinkRegistry> RegistryPtr = ObjectNameEditSinkRegistry.Pin())
+	{
+		if (Objects.Num() == 1)
+		{
+			if (const UObject* const ObjectPtr = Objects[0].Get())
+			{
+				const TSharedPtr<const UE::EditorWidgets::IObjectNameEditSink> ObjectNameEditPtr = RegistryPtr->GetObjectNameEditSinkForClass(ObjectPtr->GetClass());
+				if (ObjectNameEditPtr)
+				{
+					return ObjectNameEditPtr->GetObjectDisplayNameSuffix(ObjectPtr);
+				}
+			}
+		}
+	}
+
+	return FText();
+}
+
+
 // No code beyond this point
 #undef LOCTEXT_NAMESPACE
diff --git a/Engine/Source/Editor/EditorWidgets/Private/SObjectNameEditableTextBox.h b/Engine/Source/Editor/EditorWidgets/Private/SObjectNameEditableTextBox.h
--- a/Engine/Source/Editor/EditorWidgets/Private/SObjectNameEditableTextBox.h
+++ b/Engine/Source/Editor/EditorWidgets/Private/SObjectNameEditableTextBox.h
@@ -70,6 +70,12 @@ private:
 	/** Helper function to access name registry and retrieve display name */
 	FText GetObjectDisplayName(TWeakObjectPtr<UObject> Object) const;
 
+	/** Helper function to access name registry and determine if the name has a suffix */
+	EVisibility GetNameSuffixVisibility() const;
+
+	/** Access name registry and retrieve the name suffix, if there is one */
+	FText GetObjectDisplayNameSuffix() const;
+
 	/** The list of objects whose names are edited by the widget */
 	TArray<TWeakObjectPtr<UObject>> Objects;
 
diff --git a/Engine/Source/Editor/EditorWidgets/Public/IObjectNameEditSink.h b/Engine/Source/Editor/EditorWidgets/Public/IObjectNameEditSink.h
--- a/Engine/Source/Editor/EditorWidgets/Public/IObjectNameEditSink.h
+++ b/Engine/Source/Editor/EditorWidgets/Public/IObjectNameEditSink.h
@@ -37,6 +37,19 @@ public:
 
 	/** Returns the tooltip text for that object */
 	virtual FText GetObjectNameTooltip(UObject* Object) const = 0;
+
+	/** Whether or not the name has a suffix */
+	virtual bool ObjectDisplayNameHasSuffix(const UObject* Object) const
+	{
+		return false;
+	};
+
+	/** Returns the displayed suffix for that object. */
+	virtual FText GetObjectDisplayNameSuffix(const UObject* Object) const
+	{
+		return FText();
+	}
+
 };
 
 } //end namespace UE::EditorWidgets
