EditorUsability : AssetSearchBox
- Fixed asset search box suggestion width should match search box width instead of changing based on in view content
- Fixed content browser suggestions should appear above to avoid obfuscating filtered results
- Fixed hardcoded tooltip is used as last resort
- Exposed GetText to retrieve active content (needed for a following changelist)

#Synced-CL 49388435

diff --git a/Engine/Source/Editor/ContentBrowser/Private/SContentBrowser.cpp b/Engine/Source/Editor/ContentBrowser/Private/SContentBrowser.cpp
--- a/Engine/Source/Editor/ContentBrowser/Private/SContentBrowser.cpp
+++ b/Engine/Source/Editor/ContentBrowser/Private/SContentBrowser.cpp
@@ -1026,6 +1026,7 @@ TSharedRef<SWidget> SContentBrowser::CreateAssetView(const FContentBrowserConfig
 	LegacyContentSourceWidgets->AssetViewPtr->SetFilterBar(LegacyContentSourceWidgets->FilterListPtr);
 
 	LegacyContentSourceWidgets->SearchBoxPtr = SNew(SAssetSearchBox)
+				   .SuggestionListPlacement(EMenuPlacement::MenuPlacement_AboveAnchor) // Avoid showing suggestion on filtered content browser results...
 				   .HintText(this, &SContentBrowser::GetSearchAssetsHintText)
 				   .ShowSearchHistory(true)
 				   .OnTextChanged(this, &SContentBrowser::OnSearchBoxChanged)
diff --git a/Engine/Source/Editor/EditorWidgets/Private/SAssetSearchBox.cpp b/Engine/Source/Editor/EditorWidgets/Private/SAssetSearchBox.cpp
--- a/Engine/Source/Editor/EditorWidgets/Private/SAssetSearchBox.cpp
+++ b/Engine/Source/Editor/EditorWidgets/Private/SAssetSearchBox.cpp
@@ -97,16 +97,15 @@ void SAssetSearchBox::Construct( const FArguments& InArgs )
 				.OnKeyDownHandler(this, &SAssetSearchBox::HandleKeyDown)
 				.ShowSearchHistory(InArgs._ShowSearchHistory)
 				.OnSaveSearchClicked(InArgs._OnSaveSearchClicked)
-				.ToolTip(IDocumentation::Get()->CreateToolTip(NSLOCTEXT("AssetSearchBox", "FilterSearchHint", "Type here to search"), /*OverrideContent = */nullptr, "Shared/ContentBrowser", "AdvancedSearchSyntax"))
 			]
 			.MenuContent
-				(
+			(
 				SNew(SBorder)
 				.BorderImage(FAppStyle::GetBrush("Menu.Background"))
 				.Padding( FMargin(2.f) )
 				[
 					SNew(SBox)
-					.MinDesiredWidth(175.f)	// to enforce some minimum width, ideally we define the minimum, not a fixed width
+					.WidthOverride(this, &SAssetSearchBox::GetSuggestionListWidth)
 					.HeightOverride(250.f)	// avoids flickering, ideally this would be adaptive to the content without flickering
 					[
 						SAssignNew(SuggestionListView, SListView< TSharedPtr<FSuggestionListEntry> >)
@@ -119,6 +118,20 @@ void SAssetSearchBox::Construct( const FArguments& InArgs )
 				]
 			)
 		];
+
+	// Set Search Tooltip
+	if (InArgs._ToolTipText.IsSet())
+	{
+		InputText->SetToolTipText(InArgs._ToolTipText);
+	}
+	else if (InArgs._ToolTip.IsSet())
+	{
+		InputText->SetToolTip(InArgs._ToolTip);
+	}
+	else
+	{
+		InputText->SetToolTip(IDocumentation::Get()->CreateToolTip(NSLOCTEXT("AssetSearchBox", "FilterSearchHint", "Type here to search"), /*OverrideContent = */nullptr, "Shared/ContentBrowser", "AdvancedSearchSyntax"));
+	}
 }
 
 void SAssetSearchBox::SetText(const TAttribute< FText >& InNewText)
@@ -127,6 +140,11 @@ void SAssetSearchBox::SetText(const TAttribute< FText >& InNewText)
 	PreCommittedText = InNewText.Get();
 }
 
+FText SAssetSearchBox::GetText() const
+{
+	return InputText->GetText();
+}
+
 void SAssetSearchBox::SetError( const FText& InError )
 {
 	InputText->SetError(InError);
@@ -432,6 +450,19 @@ FText SAssetSearchBox::DefaultSuggestionChosenImpl(const FText& SearchText, cons
 	return FText::FromString(Suggestion);
 }
 
+FOptionalSize SAssetSearchBox::GetSuggestionListWidth() const
+{
+	// to enforce some minimum width, ideally we define the minimum, not a fixed width
+	constexpr float MinWidth = 175.f;
+
+	if (!InputText.IsValid())
+	{
+		return MinWidth;
+	}
+
+	return FMath::Max(InputText->GetCachedGeometry().GetLocalSize().X, MinWidth);
+}
+
 void SAssetSearchBox::SetOnSaveSearchHandler(SFilterSearchBox::FOnSaveSearchClicked InOnSaveSearchHandler)
 {
 	InputText->SetOnSaveSearchHandler(InOnSaveSearchHandler);
diff --git a/Engine/Source/Editor/EditorWidgets/Public/SAssetSearchBox.h b/Engine/Source/Editor/EditorWidgets/Public/SAssetSearchBox.h
--- a/Engine/Source/Editor/EditorWidgets/Public/SAssetSearchBox.h
+++ b/Engine/Source/Editor/EditorWidgets/Public/SAssetSearchBox.h
@@ -130,6 +130,9 @@ public:
 
 	/** Sets the text string currently being edited */
 	UE_API void SetText(const TAttribute< FText >& InNewText);
+	
+	/** Get the text string currently being edited */
+	UE_API FText GetText() const;
 
 	/** Set or clear the current error reporting information for this search box */
 	UE_API void SetError( const FText& InError );
@@ -185,7 +188,9 @@ private:
 	/** Default implementation of OnAssetSearchBoxSuggestionChosen, if no external implementation is provided */
 	static UE_API FText DefaultSuggestionChosenImpl(const FText& SearchText, const FString& Suggestion);
 
-private:
+	/** Returns the width for the suggestion list */
+	UE_API FOptionalSize GetSuggestionListWidth() const;
+
 	/** The editable text field */
 	TSharedPtr< SFilterSearchBox > InputText;
 
